[GLTFExporter] Remove GLTF exporter from Sandbox: Phase #1

commit 0de342e7473d5986692a61046e62d7fe249b34cc
Date: Mon, 18 Jul 2022 17:09:31 -0400

    Remove GLTF exporter from Sandbox: Phase #1
    preflight none

#skipci

diff --git a/Engine/Plugins/Enterprise/GLTFExporter/MigratedFromGit/GLTFExporter.uplugin b/Engine/Plugins/Enterprise/GLTFExporter/MigratedFromGit/GLTFExporter.uplugin
--- a/Engine/Plugins/Enterprise/GLTFExporter/MigratedFromGit/GLTFExporter.uplugin
+++ b/Engine/Plugins/Enterprise/GLTFExporter/MigratedFromGit/GLTFExporter.uplugin
@@ -10,7 +10,6 @@
 	"DocsURL": "",
 	"MarketplaceURL": "",
 	"SupportURL": "",
-	"EngineVersion": "5.0.0",
 	"CanContainContent": true,
 	"Modules": [
 		{
diff --git a/Engine/Plugins/Enterprise/GLTFExporter/MigratedFromGit/Source/GLTFExporter/Private/Actors/GLTFHotspotActor.cpp b/Engine/Plugins/Enterprise/GLTFExporter/MigratedFromGit/Source/GLTFExporter/Private/Actors/GLTFHotspotActor.cpp
--- a/Engine/Plugins/Enterprise/GLTFExporter/MigratedFromGit/Source/GLTFExporter/Private/Actors/GLTFHotspotActor.cpp
+++ b/Engine/Plugins/Enterprise/GLTFExporter/MigratedFromGit/Source/GLTFExporter/Private/Actors/GLTFHotspotActor.cpp
@@ -453,7 +453,7 @@ void AGLTFHotspotActor::ValidateAnimation()
 {
 	if (SkeletalMeshActor != nullptr && AnimationSequence != nullptr)
 	{
-		if (const USkeletalMesh* SkeletalMesh = SkeletalMeshActor->GetSkeletalMeshComponent()->SkeletalMesh)
+		if (const USkeletalMesh* SkeletalMesh = SkeletalMeshActor->GetSkeletalMeshComponent()->GetSkeletalMesh())
 		{
 #if (ENGINE_MAJOR_VERSION > 4 || ENGINE_MINOR_VERSION >= 27)
 			const USkeleton* Skeleton = SkeletalMesh->GetSkeleton();
diff --git a/Engine/Plugins/Enterprise/GLTFExporter/MigratedFromGit/Source/GLTFExporter/Private/Builders/GLTFConvertBuilder.cpp b/Engine/Plugins/Enterprise/GLTFExporter/MigratedFromGit/Source/GLTFExporter/Private/Builders/GLTFConvertBuilder.cpp
--- a/Engine/Plugins/Enterprise/GLTFExporter/MigratedFromGit/Source/GLTFExporter/Private/Builders/GLTFConvertBuilder.cpp
+++ b/Engine/Plugins/Enterprise/GLTFExporter/MigratedFromGit/Source/GLTFExporter/Private/Builders/GLTFConvertBuilder.cpp
@@ -2,6 +2,7 @@
 
 #include "Builders/GLTFConvertBuilder.h"
 #include "Converters/GLTFMeshUtility.h"
+#include "Components/SkeletalMeshComponent.h"
 #include "Engine/Texture2D.h"
 #include "Engine/TextureCube.h"
 #include "Engine/TextureRenderTarget2D.h"
@@ -162,7 +163,7 @@ FGLTFJsonMesh* FGLTFConvertBuilder::AddUniqueMesh(const USkeletalMeshComponent*
 		return nullptr;
 	}
 
-	const USkeletalMesh* SkeletalMesh = SkeletalMeshComponent->SkeletalMesh;
+	const USkeletalMesh* SkeletalMesh = SkeletalMeshComponent->GetSkeletalMesh();
 	if (SkeletalMesh == nullptr)
 	{
 		return nullptr;
@@ -224,7 +225,7 @@ FGLTFJsonMaterial* FGLTFConvertBuilder::AddUniqueMaterial(const UMaterialInterfa
 FGLTFJsonMaterial* FGLTFConvertBuilder::AddUniqueMaterial(const UMaterialInterface* Material, const USkeletalMeshComponent* SkeletalMeshComponent, int32 LODIndex, int32 MaterialIndex)
 {
 	// TODO: optimize by skipping mesh data if material doesn't need it
-	const USkeletalMesh* SkeletalMesh = SkeletalMeshComponent->SkeletalMesh;
+	const USkeletalMesh* SkeletalMesh = SkeletalMeshComponent->GetSkeletalMesh();
 	const FGLTFMeshData* MeshData = AddUniqueMeshData(SkeletalMesh, SkeletalMeshComponent, LODIndex);
 	const FGLTFIndexArray SectionIndices = FGLTFMeshUtility::GetSectionIndices(SkeletalMesh, MeshData->LODIndex, MaterialIndex);
 	return AddUniqueMaterial(Material, MeshData, SectionIndices);
@@ -362,7 +363,7 @@ FGLTFJsonSkin* FGLTFConvertBuilder::AddUniqueSkin(FGLTFJsonNode* RootNode, const
 		return nullptr;
 	}
 
-	const USkeletalMesh* SkeletalMesh = SkeletalMeshComponent->SkeletalMesh;
+	const USkeletalMesh* SkeletalMesh = SkeletalMeshComponent->GetSkeletalMesh();
 
 	if (SkeletalMesh == nullptr)
 	{
diff --git a/Engine/Plugins/Enterprise/GLTFExporter/MigratedFromGit/Source/GLTFExporter/Private/Converters/GLTFAccessorConverters.cpp b/Engine/Plugins/Enterprise/GLTFExporter/MigratedFromGit/Source/GLTFExporter/Private/Converters/GLTFAccessorConverters.cpp
--- a/Engine/Plugins/Enterprise/GLTFExporter/MigratedFromGit/Source/GLTFExporter/Private/Converters/GLTFAccessorConverters.cpp
+++ b/Engine/Plugins/Enterprise/GLTFExporter/MigratedFromGit/Source/GLTFExporter/Private/Converters/GLTFAccessorConverters.cpp
@@ -5,6 +5,14 @@
 #include "Converters/GLTFBufferAdapter.h"
 #include "Builders/GLTFConvertBuilder.h"
 
+#include "Rendering/ColorVertexBuffer.h"
+#include "Rendering/PositionVertexBuffer.h"
+#include "RawIndexBuffer.h"
+#include "RHI.h"
+#include "RHIResources.h"
+#include "Rendering/SkinWeightVertexBuffer.h"
+#include "Rendering/StaticMeshVertexBuffer.h"
+
 // TODO: Unreal-style implementation of std::conditional to avoid mixing in STL. Should be added to the engine.
 template <bool Condition, class TypeIfTrue, class TypeIfFalse>
 class TConditional
diff --git a/Engine/Plugins/Enterprise/GLTFExporter/MigratedFromGit/Source/GLTFExporter/Private/Converters/GLTFAnimationConverters.cpp b/Engine/Plugins/Enterprise/GLTFExporter/MigratedFromGit/Source/GLTFExporter/Private/Converters/GLTFAnimationConverters.cpp
--- a/Engine/Plugins/Enterprise/GLTFExporter/MigratedFromGit/Source/GLTFExporter/Private/Converters/GLTFAnimationConverters.cpp
+++ b/Engine/Plugins/Enterprise/GLTFExporter/MigratedFromGit/Source/GLTFExporter/Private/Converters/GLTFAnimationConverters.cpp
@@ -3,7 +3,7 @@
 #include "Converters/GLTFAnimationConverters.h"
 #include "Builders/GLTFConvertBuilder.h"
 #include "Tasks/GLTFDelayedAnimationTasks.h"
-#include "Animation/AnimSequence.h"
+#include "Components/SkeletalMeshComponent.h"
 #include "LevelSequenceActor.h"
 
 FGLTFJsonAnimation* FGLTFAnimationConverter::Convert(FGLTFJsonNode* RootNode, const USkeletalMesh* SkeletalMesh, const UAnimSequence* AnimSequence)
@@ -40,7 +40,7 @@ FGLTFJsonAnimation* FGLTFAnimationConverter::Convert(FGLTFJsonNode* RootNode, co
 
 FGLTFJsonAnimation* FGLTFAnimationDataConverter::Convert(FGLTFJsonNode* RootNode, const USkeletalMeshComponent* SkeletalMeshComponent)
 {
-	const USkeletalMesh* SkeletalMesh = SkeletalMeshComponent->SkeletalMesh;
+	const USkeletalMesh* SkeletalMesh = SkeletalMeshComponent->GetSkeletalMesh();
 	if (SkeletalMesh == nullptr)
 	{
 		return nullptr;
diff --git a/Engine/Plugins/Enterprise/GLTFExporter/MigratedFromGit/Source/GLTFExporter/Private/Converters/GLTFBlueprintUtility.h b/Engine/Plugins/Enterprise/GLTFExporter/MigratedFromGit/Source/GLTFExporter/Private/Converters/GLTFBlueprintUtility.h
--- a/Engine/Plugins/Enterprise/GLTFExporter/MigratedFromGit/Source/GLTFExporter/Private/Converters/GLTFBlueprintUtility.h
+++ b/Engine/Plugins/Enterprise/GLTFExporter/MigratedFromGit/Source/GLTFExporter/Private/Converters/GLTFBlueprintUtility.h
@@ -3,6 +3,10 @@
 #pragma once
 
 #include "CoreMinimal.h"
+#include "UObject/Object.h"
+#include "UObject/UnrealType.h"
+
+class AActor;
 
 struct FGLTFBlueprintUtility
 {
diff --git a/Engine/Plugins/Enterprise/GLTFExporter/MigratedFromGit/Source/GLTFExporter/Private/Converters/GLTFBoneUtility.cpp b/Engine/Plugins/Enterprise/GLTFExporter/MigratedFromGit/Source/GLTFExporter/Private/Converters/GLTFBoneUtility.cpp
--- a/Engine/Plugins/Enterprise/GLTFExporter/MigratedFromGit/Source/GLTFExporter/Private/Converters/GLTFBoneUtility.cpp
+++ b/Engine/Plugins/Enterprise/GLTFExporter/MigratedFromGit/Source/GLTFExporter/Private/Converters/GLTFBoneUtility.cpp
@@ -3,6 +3,9 @@
 #include "Converters/GLTFBoneUtility.h"
 #include "Sections/MovieScene3DTransformSection.h"
 #include "Animation/AnimationPoseData.h"
+#include "Animation/AnimSequence.h"
+#include "Animation/AttributesRuntime.h"
+#include "ReferenceSkeleton.h"
 
 FTransform FGLTFBoneUtility::GetBindTransform(const FReferenceSkeleton& RefSkeleton, int32 BoneIndex)
 {
diff --git a/Engine/Plugins/Enterprise/GLTFExporter/MigratedFromGit/Source/GLTFExporter/Private/Converters/GLTFBoneUtility.h b/Engine/Plugins/Enterprise/GLTFExporter/MigratedFromGit/Source/GLTFExporter/Private/Converters/GLTFBoneUtility.h
--- a/Engine/Plugins/Enterprise/GLTFExporter/MigratedFromGit/Source/GLTFExporter/Private/Converters/GLTFBoneUtility.h
+++ b/Engine/Plugins/Enterprise/GLTFExporter/MigratedFromGit/Source/GLTFExporter/Private/Converters/GLTFBoneUtility.h
@@ -3,9 +3,13 @@
 #pragma once
 
 #include "CoreMinimal.h"
+#include "BoneIndices.h"
 
 struct FMovieSceneDoubleChannel;
+struct FReferenceSkeleton;
+class UAnimSequence;
 class UMovieScene3DTransformSection;
+class USkeleton;
 
 struct FGLTFBoneUtility
 {
diff --git a/Engine/Plugins/Enterprise/GLTFExporter/MigratedFromGit/Source/GLTFExporter/Private/Converters/GLTFBufferAdapter.cpp b/Engine/Plugins/Enterprise/GLTFExporter/MigratedFromGit/Source/GLTFExporter/Private/Converters/GLTFBufferAdapter.cpp
--- a/Engine/Plugins/Enterprise/GLTFExporter/MigratedFromGit/Source/GLTFExporter/Private/Converters/GLTFBufferAdapter.cpp
+++ b/Engine/Plugins/Enterprise/GLTFExporter/MigratedFromGit/Source/GLTFExporter/Private/Converters/GLTFBufferAdapter.cpp
@@ -3,6 +3,14 @@
 #include "Converters/GLTFBufferAdapter.h"
 #include "Converters/GLTFBufferUtility.h"
 
+#include "Rendering/ColorVertexBuffer.h"
+#include "Rendering/PositionVertexBuffer.h"
+#include "RawIndexBuffer.h"
+#include "RHI.h"
+#include "RHIResources.h"
+#include "Rendering/SkinWeightVertexBuffer.h"
+#include "Rendering/StaticMeshVertexBuffer.h"
+
 class FGLTFBufferAdapterCPU final : public IGLTFBufferAdapter
 {
 public:
diff --git a/Engine/Plugins/Enterprise/GLTFExporter/MigratedFromGit/Source/GLTFExporter/Private/Converters/GLTFBufferAdapter.h b/Engine/Plugins/Enterprise/GLTFExporter/MigratedFromGit/Source/GLTFExporter/Private/Converters/GLTFBufferAdapter.h
--- a/Engine/Plugins/Enterprise/GLTFExporter/MigratedFromGit/Source/GLTFExporter/Private/Converters/GLTFBufferAdapter.h
+++ b/Engine/Plugins/Enterprise/GLTFExporter/MigratedFromGit/Source/GLTFExporter/Private/Converters/GLTFBufferAdapter.h
@@ -4,6 +4,13 @@
 
 #include "CoreMinimal.h"
 
+class FRawStaticIndexBuffer;
+class FRawStaticIndexBuffer16or32Interface;
+class FPositionVertexBuffer;
+class FColorVertexBuffer;
+class FStaticMeshVertexBuffer;
+class FSkinWeightVertexBuffer;
+
 class IGLTFBufferAdapter
 {
 public:
diff --git a/Engine/Plugins/Enterprise/GLTFExporter/MigratedFromGit/Source/GLTFExporter/Private/Converters/GLTFBufferUtility.h b/Engine/Plugins/Enterprise/GLTFExporter/MigratedFromGit/Source/GLTFExporter/Private/Converters/GLTFBufferUtility.h
--- a/Engine/Plugins/Enterprise/GLTFExporter/MigratedFromGit/Source/GLTFExporter/Private/Converters/GLTFBufferUtility.h
+++ b/Engine/Plugins/Enterprise/GLTFExporter/MigratedFromGit/Source/GLTFExporter/Private/Converters/GLTFBufferUtility.h
@@ -4,6 +4,16 @@
 
 #include "CoreMinimal.h"
 
+class FColorVertexBuffer;
+class FPositionVertexBuffer;
+class FRawStaticIndexBuffer;
+class FRawStaticIndexBuffer16or32Interface;
+class FRHIBuffer;
+class FSkinWeightVertexBuffer;
+class FSkinWeightDataVertexBuffer;
+class FSkinWeightLookupVertexBuffer;
+class FStaticMeshVertexBuffer;
+
 struct FGLTFBufferUtility
 {
 	static bool HasCPUAccess(const FRawStaticIndexBuffer* IndexBuffer);
diff --git a/Engine/Plugins/Enterprise/GLTFExporter/MigratedFromGit/Source/GLTFExporter/Private/Converters/GLTFCurveUtility.h b/Engine/Plugins/Enterprise/GLTFExporter/MigratedFromGit/Source/GLTFExporter/Private/Converters/GLTFCurveUtility.h
--- a/Engine/Plugins/Enterprise/GLTFExporter/MigratedFromGit/Source/GLTFExporter/Private/Converters/GLTFCurveUtility.h
+++ b/Engine/Plugins/Enterprise/GLTFExporter/MigratedFromGit/Source/GLTFExporter/Private/Converters/GLTFCurveUtility.h
@@ -2,6 +2,8 @@
 
 #pragma once
 
+#include "Math/UnrealMathUtility.h"
+
 class UCurveLinearColor;
 
 struct FGLTFCurveUtility
diff --git a/Engine/Plugins/Enterprise/GLTFExporter/MigratedFromGit/Source/GLTFExporter/Private/Converters/GLTFHotspotConverters.cpp b/Engine/Plugins/Enterprise/GLTFExporter/MigratedFromGit/Source/GLTFExporter/Private/Converters/GLTFHotspotConverters.cpp
--- a/Engine/Plugins/Enterprise/GLTFExporter/MigratedFromGit/Source/GLTFExporter/Private/Converters/GLTFHotspotConverters.cpp
+++ b/Engine/Plugins/Enterprise/GLTFExporter/MigratedFromGit/Source/GLTFExporter/Private/Converters/GLTFHotspotConverters.cpp
@@ -3,6 +3,7 @@
 #include "Converters/GLTFHotspotConverters.h"
 #include "Builders/GLTFContainerBuilder.h"
 #include "Animation/SkeletalMeshActor.h"
+#include "Components/SkeletalMeshComponent.h"
 
 FGLTFJsonHotspot* FGLTFHotspotConverter::Convert(const AGLTFHotspotActor* HotspotActor)
 {
@@ -28,7 +29,7 @@ FGLTFJsonHotspot* FGLTFHotspotConverter::Convert(const AGLTFHotspotActor* Hotspo
 			FGLTFJsonNode* RootNode = Builder.AddUniqueNode(SkeletalMeshActor);
 
 			const USkeletalMeshComponent* SkeletalMeshComponent = SkeletalMeshActor->GetSkeletalMeshComponent();
-			if (const USkeletalMesh* SkeletalMesh = SkeletalMeshComponent->SkeletalMesh)
+			if (const USkeletalMesh* SkeletalMesh = SkeletalMeshComponent->GetSkeletalMesh())
 			{
 				if (const UAnimSequence* AnimSequence = HotspotActor->AnimationSequence)
 				{
diff --git a/Engine/Plugins/Enterprise/GLTFExporter/MigratedFromGit/Source/GLTFExporter/Private/Converters/GLTFKhrVariantConverters.cpp b/Engine/Plugins/Enterprise/GLTFExporter/MigratedFromGit/Source/GLTFExporter/Private/Converters/GLTFKhrVariantConverters.cpp
--- a/Engine/Plugins/Enterprise/GLTFExporter/MigratedFromGit/Source/GLTFExporter/Private/Converters/GLTFKhrVariantConverters.cpp
+++ b/Engine/Plugins/Enterprise/GLTFExporter/MigratedFromGit/Source/GLTFExporter/Private/Converters/GLTFKhrVariantConverters.cpp
@@ -3,6 +3,7 @@
 #include "Converters/GLTFKhrVariantConverters.h"
 #include "Converters/GLTFVariantUtility.h"
 #include "Builders/GLTFContainerBuilder.h"
+#include "Components/MeshComponent.h"
 #include "VariantObjectBinding.h"
 #include "PropertyValueMaterial.h"
 #include "PropertyValue.h"
diff --git a/Engine/Plugins/Enterprise/GLTFExporter/MigratedFromGit/Source/GLTFExporter/Private/Converters/GLTFMaterialProxyFactory.cpp b/Engine/Plugins/Enterprise/GLTFExporter/MigratedFromGit/Source/GLTFExporter/Private/Converters/GLTFMaterialProxyFactory.cpp
--- a/Engine/Plugins/Enterprise/GLTFExporter/MigratedFromGit/Source/GLTFExporter/Private/Converters/GLTFMaterialProxyFactory.cpp
+++ b/Engine/Plugins/Enterprise/GLTFExporter/MigratedFromGit/Source/GLTFExporter/Private/Converters/GLTFMaterialProxyFactory.cpp
@@ -9,6 +9,7 @@
 #include "Materials/GLTFProxyMaterialInfo.h"
 #include "Materials/MaterialInstanceConstant.h"
 #include "AssetRegistry/AssetRegistryModule.h"
+#include "HAL/FileManager.h"
 #include "ImageUtils.h"
 
 FGLTFMaterialProxyFactory::FGLTFMaterialProxyFactory(const UGLTFProxyOptions* Options)
diff --git a/Engine/Plugins/Enterprise/GLTFExporter/MigratedFromGit/Source/GLTFExporter/Private/Converters/GLTFMaterialUtility.h b/Engine/Plugins/Enterprise/GLTFExporter/MigratedFromGit/Source/GLTFExporter/Private/Converters/GLTFMaterialUtility.h
--- a/Engine/Plugins/Enterprise/GLTFExporter/MigratedFromGit/Source/GLTFExporter/Private/Converters/GLTFMaterialUtility.h
+++ b/Engine/Plugins/Enterprise/GLTFExporter/MigratedFromGit/Source/GLTFExporter/Private/Converters/GLTFMaterialUtility.h
@@ -4,10 +4,13 @@
 
 #include "Builders/GLTFContainerBuilder.h"
 #include "Converters/GLTFSharedArray.h"
+#include "MaterialExpressionIO.h"
 
 struct FGLTFMaterialAnalysis;
 struct FGLTFMaterialPropertyEx;
+class UMaterialExpressionTextureSample;
 class UMaterialInstanceConstant;
+class UMaterialInterface;
 
 struct FGLTFPropertyBakeOutput
 {
diff --git a/Engine/Plugins/Enterprise/GLTFExporter/MigratedFromGit/Source/GLTFExporter/Private/Converters/GLTFMeshData.cpp b/Engine/Plugins/Enterprise/GLTFExporter/MigratedFromGit/Source/GLTFExporter/Private/Converters/GLTFMeshData.cpp
--- a/Engine/Plugins/Enterprise/GLTFExporter/MigratedFromGit/Source/GLTFExporter/Private/Converters/GLTFMeshData.cpp
+++ b/Engine/Plugins/Enterprise/GLTFExporter/MigratedFromGit/Source/GLTFExporter/Private/Converters/GLTFMeshData.cpp
@@ -3,10 +3,15 @@
 #include "Converters/GLTFMeshData.h"
 #include "Converters/GLTFNameUtility.h"
 
+#include "Components/SkeletalMeshComponent.h"
+#include "Components/StaticMeshComponent.h"
+#include "Engine/StaticMesh.h"
+#include "Engine/SkeletalMesh.h"
 #if WITH_EDITOR
-#include "StaticMeshAttributes.h"
-#include "Engine/MapBuildDataRegistry.h"
 #include "Developer/MeshMergeUtilities/Private/MeshMergeHelpers.h"
+#include "Editor.h"
+#include "Engine/MapBuildDataRegistry.h"
+#include "StaticMeshAttributes.h"
 #include "Rendering/SkeletalMeshRenderData.h"
 #endif
 
@@ -87,7 +92,7 @@ FGLTFMeshData::FGLTFMeshData(const USkeletalMesh* SkeletalMesh, const USkeletalM
 
 #if WITH_EDITOR
 		PrimitiveData = { SkeletalMeshComponent };
-		FMeshMergeHelpers::RetrieveMesh(const_cast<USkeletalMeshComponent*>(SkeletalMeshComponent), LODIndex, Description, true);
+		FMeshMergeHelpers::RetrieveMesh(SkeletalMeshComponent, LODIndex, Description, true);
 #endif
 	}
 	else
diff --git a/Engine/Plugins/Enterprise/GLTFExporter/MigratedFromGit/Source/GLTFExporter/Private/Converters/GLTFMeshUtility.cpp b/Engine/Plugins/Enterprise/GLTFExporter/MigratedFromGit/Source/GLTFExporter/Private/Converters/GLTFMeshUtility.cpp
--- a/Engine/Plugins/Enterprise/GLTFExporter/MigratedFromGit/Source/GLTFExporter/Private/Converters/GLTFMeshUtility.cpp
+++ b/Engine/Plugins/Enterprise/GLTFExporter/MigratedFromGit/Source/GLTFExporter/Private/Converters/GLTFMeshUtility.cpp
@@ -1,9 +1,13 @@
 // Copyright Epic Games, Inc. All Rights Reserved.
 
 #include "Converters/GLTFMeshUtility.h"
+#include "Components/SkeletalMeshComponent.h"
+#include "Engine/StaticMesh.h"
+#include "Engine/SkeletalMesh.h"
+#include "Materials/Material.h"
 #include "Rendering/SkeletalMeshRenderData.h"
 #include "StaticMeshCompiler.h"
-#include "SkeletalMeshCompiler.h"
+#include "SkinnedAssetCompiler.h"
 
 void FGLTFMeshUtility::FullyLoad(const UStaticMesh* InStaticMesh)
 {
@@ -22,7 +26,7 @@ void FGLTFMeshUtility::FullyLoad(const USkeletalMesh* InSkeletalMesh)
 	USkeletalMesh* SkeletalMesh = const_cast<USkeletalMesh*>(InSkeletalMesh);
 
 #if WITH_EDITOR
-	FSkeletalMeshCompilingManager::Get().FinishCompilation({ SkeletalMesh });
+	FSkinnedAssetCompilingManager::Get().FinishCompilation({ SkeletalMesh });
 #endif
 
 	SkeletalMesh->SetForceMipLevelsToBeResident(30.0f);
diff --git a/Engine/Plugins/Enterprise/GLTFExporter/MigratedFromGit/Source/GLTFExporter/Private/Converters/GLTFMeshUtility.h b/Engine/Plugins/Enterprise/GLTFExporter/MigratedFromGit/Source/GLTFExporter/Private/Converters/GLTFMeshUtility.h
--- a/Engine/Plugins/Enterprise/GLTFExporter/MigratedFromGit/Source/GLTFExporter/Private/Converters/GLTFMeshUtility.h
+++ b/Engine/Plugins/Enterprise/GLTFExporter/MigratedFromGit/Source/GLTFExporter/Private/Converters/GLTFMeshUtility.h
@@ -4,7 +4,15 @@
 
 #include "Converters/GLTFIndexArray.h"
 
+class UMaterialInterface;
+class UStaticMesh;
+class UStaticMeshComponent;
+class USkeletalMesh;
+class USkeletalMeshComponent;
 class FSkeletalMeshLODRenderData;
+struct FStaticMaterial;
+struct FSkeletalMaterial;
+struct FStaticMeshLODResources;
 
 struct FGLTFMeshUtility
 {
diff --git a/Engine/Plugins/Enterprise/GLTFExporter/MigratedFromGit/Source/GLTFExporter/Private/Converters/GLTFNameUtility.cpp b/Engine/Plugins/Enterprise/GLTFExporter/MigratedFromGit/Source/GLTFExporter/Private/Converters/GLTFNameUtility.cpp
--- a/Engine/Plugins/Enterprise/GLTFExporter/MigratedFromGit/Source/GLTFExporter/Private/Converters/GLTFNameUtility.cpp
+++ b/Engine/Plugins/Enterprise/GLTFExporter/MigratedFromGit/Source/GLTFExporter/Private/Converters/GLTFNameUtility.cpp
@@ -4,6 +4,8 @@
 #include "Engine/StaticMeshActor.h"
 #include "Animation/SkeletalMeshActor.h"
 #include "Components/LightComponent.h"
+#include "Components/SkeletalMeshComponent.h"
+#include "Components/StaticMeshComponent.h"
 #include "Engine/Light.h"
 #include "Camera/CameraComponent.h"
 #include "Camera/CameraActor.h"
diff --git a/Engine/Plugins/Enterprise/GLTFExporter/MigratedFromGit/Source/GLTFExporter/Private/Converters/GLTFNameUtility.h b/Engine/Plugins/Enterprise/GLTFExporter/MigratedFromGit/Source/GLTFExporter/Private/Converters/GLTFNameUtility.h
--- a/Engine/Plugins/Enterprise/GLTFExporter/MigratedFromGit/Source/GLTFExporter/Private/Converters/GLTFNameUtility.h
+++ b/Engine/Plugins/Enterprise/GLTFExporter/MigratedFromGit/Source/GLTFExporter/Private/Converters/GLTFNameUtility.h
@@ -3,11 +3,19 @@
 #pragma once
 
 #include "CoreMinimal.h"
+#include "UObject/Class.h"
+#include "UObject/ReflectedTypeAccessors.h"
 
 class ULightComponent;
 class UCameraComponent;
 class AGLTFHotspotActor;
 
+class UCameraComponent;
+class ULightComponent;
+class USceneComponent;
+class USkeletalMeshComponent;
+class UStaticMeshComponent;
+
 struct FGLTFNameUtility
 {
 	template <typename EnumType, typename = typename TEnableIf<TIsEnum<EnumType>::Value>::Type>
diff --git a/Engine/Plugins/Enterprise/GLTFExporter/MigratedFromGit/Source/GLTFExporter/Private/Converters/GLTFNodeConverters.cpp b/Engine/Plugins/Enterprise/GLTFExporter/MigratedFromGit/Source/GLTFExporter/Private/Converters/GLTFNodeConverters.cpp
--- a/Engine/Plugins/Enterprise/GLTFExporter/MigratedFromGit/Source/GLTFExporter/Private/Converters/GLTFNodeConverters.cpp
+++ b/Engine/Plugins/Enterprise/GLTFExporter/MigratedFromGit/Source/GLTFExporter/Private/Converters/GLTFNodeConverters.cpp
@@ -9,6 +9,8 @@
 #include "LevelSequenceActor.h"
 #include "Camera/CameraComponent.h"
 #include "Components/LightComponent.h"
+#include "Components/SkeletalMeshComponent.h"
+#include "Engine/StaticMesh.h"
 #include "Engine/StaticMeshSocket.h"
 #include "Engine/SkeletalMeshSocket.h"
 
@@ -221,8 +223,11 @@ FGLTFJsonNode* FGLTFComponentSocketConverter::Convert(const USceneComponent* Sce
 		if (const USkinnedMeshComponent* SkinnedMeshComponent = Cast<USkinnedMeshComponent>(SceneComponent))
 		{
 			// TODO: add support for SocketOverrideLookup?
-			const USkeletalMesh* SkeletalMesh = SkinnedMeshComponent->SkeletalMesh;
-			return Builder.AddUniqueNode(JsonNode, SkeletalMesh, SocketName);
+			// #ueent_correction_to_review
+			if (const USkeletalMesh* SkeletalMesh = Cast<USkeletalMesh>(SkinnedMeshComponent->GetSkinnedAsset()))
+			{
+				return Builder.AddUniqueNode(JsonNode, SkeletalMesh, SocketName);
+			}
 		}
 
 		// TODO: add support for more socket types
diff --git a/Engine/Plugins/Enterprise/GLTFExporter/MigratedFromGit/Source/GLTFExporter/Private/Converters/GLTFSkinConverters.cpp b/Engine/Plugins/Enterprise/GLTFExporter/MigratedFromGit/Source/GLTFExporter/Private/Converters/GLTFSkinConverters.cpp
--- a/Engine/Plugins/Enterprise/GLTFExporter/MigratedFromGit/Source/GLTFExporter/Private/Converters/GLTFSkinConverters.cpp
+++ b/Engine/Plugins/Enterprise/GLTFExporter/MigratedFromGit/Source/GLTFExporter/Private/Converters/GLTFSkinConverters.cpp
@@ -5,6 +5,10 @@
 #include "Converters/GLTFBoneUtility.h"
 #include "Builders/GLTFConvertBuilder.h"
 
+#include "Animation/Skeleton.h"
+#include "Components/SkeletalMeshComponent.h"
+#include "Engine/SkeletalMesh.h"
+
 FGLTFJsonSkin* FGLTFSkinConverter::Convert(FGLTFJsonNode* RootNode, const USkeletalMesh* SkeletalMesh)
 {
 #if (ENGINE_MAJOR_VERSION > 4 || ENGINE_MINOR_VERSION >= 27)
diff --git a/Engine/Plugins/Enterprise/GLTFExporter/MigratedFromGit/Source/GLTFExporter/Private/Converters/GLTFSkySphereConverters.cpp b/Engine/Plugins/Enterprise/GLTFExporter/MigratedFromGit/Source/GLTFExporter/Private/Converters/GLTFSkySphereConverters.cpp
--- a/Engine/Plugins/Enterprise/GLTFExporter/MigratedFromGit/Source/GLTFExporter/Private/Converters/GLTFSkySphereConverters.cpp
+++ b/Engine/Plugins/Enterprise/GLTFExporter/MigratedFromGit/Source/GLTFExporter/Private/Converters/GLTFSkySphereConverters.cpp
@@ -6,6 +6,7 @@
 #include "Converters/GLTFBlueprintUtility.h"
 #include "Converters/GLTFCurveUtility.h"
 #include "Builders/GLTFContainerBuilder.h"
+#include "Components/StaticMeshComponent.h"
 #include "Engine/DirectionalLight.h"
 #include "Curves/CurveLinearColor.h"
 #include "Materials/MaterialInstance.h"
diff --git a/Engine/Plugins/Enterprise/GLTFExporter/MigratedFromGit/Source/GLTFExporter/Private/Converters/GLTFTextureUtility.h b/Engine/Plugins/Enterprise/GLTFExporter/MigratedFromGit/Source/GLTFExporter/Private/Converters/GLTFTextureUtility.h
--- a/Engine/Plugins/Enterprise/GLTFExporter/MigratedFromGit/Source/GLTFExporter/Private/Converters/GLTFTextureUtility.h
+++ b/Engine/Plugins/Enterprise/GLTFExporter/MigratedFromGit/Source/GLTFExporter/Private/Converters/GLTFTextureUtility.h
@@ -3,6 +3,13 @@
 #pragma once
 
 #include "Json/GLTFJsonEnums.h"
+#include "Engine/Texture.h"
+#include "PixelFormat.h"
+#include "RHIDefinitions.h"
+
+class UTextureCube;
+class UTextureRenderTarget2D;
+class UTextureRenderTargetCube;
 
 struct FGLTFTextureUtility
 {
diff --git a/Engine/Plugins/Enterprise/GLTFExporter/MigratedFromGit/Source/GLTFExporter/Private/Converters/GLTFUVOverlapChecker.cpp b/Engine/Plugins/Enterprise/GLTFExporter/MigratedFromGit/Source/GLTFExporter/Private/Converters/GLTFUVOverlapChecker.cpp
--- a/Engine/Plugins/Enterprise/GLTFExporter/MigratedFromGit/Source/GLTFExporter/Private/Converters/GLTFUVOverlapChecker.cpp
+++ b/Engine/Plugins/Enterprise/GLTFExporter/MigratedFromGit/Source/GLTFExporter/Private/Converters/GLTFUVOverlapChecker.cpp
@@ -7,6 +7,9 @@
 #include "Modules/ModuleManager.h"
 #include "IGLTFMaterialBakingModule.h"
 #include "GLTFMaterialBakingStructures.h"
+#else
+#include "Materials/MaterialInterface.h"
+#include "UObject/UObjectGlobals.h"
 #endif
 
 void FGLTFUVOverlapChecker::Sanitize(const FMeshDescription*& Description, FGLTFIndexArray& SectionIndices, int32& TexCoord)
diff --git a/Engine/Plugins/Enterprise/GLTFExporter/MigratedFromGit/Source/GLTFExporter/Private/Converters/GLTFVariantUtility.cpp b/Engine/Plugins/Enterprise/GLTFExporter/MigratedFromGit/Source/GLTFExporter/Private/Converters/GLTFVariantUtility.cpp
--- a/Engine/Plugins/Enterprise/GLTFExporter/MigratedFromGit/Source/GLTFExporter/Private/Converters/GLTFVariantUtility.cpp
+++ b/Engine/Plugins/Enterprise/GLTFExporter/MigratedFromGit/Source/GLTFExporter/Private/Converters/GLTFVariantUtility.cpp
@@ -10,21 +10,7 @@
 
 const TArray<FCapturedPropSegment>& FGLTFVariantUtility::GetCapturedPropSegments(const UPropertyValue* Property)
 {
-#if (ENGINE_MAJOR_VERSION > 4 || ENGINE_MINOR_VERSION >= 26)
 	return Property->GetCapturedPropSegments();
-
-#else
-	class UPropertyValueHack : public UPropertyValue
-	{
-	public:
-		const TArray<FCapturedPropSegment>& GetCapturedPropSegments() const
-		{
-			return CapturedPropSegments;
-		}
-	};
-
-	return static_cast<const UPropertyValueHack*>(Property)->GetCapturedPropSegments();
-#endif
 }
 
 bool FGLTFVariantUtility::IsStaticMeshProperty(const UPropertyValue* Property)
diff --git a/Engine/Plugins/Enterprise/GLTFExporter/MigratedFromGit/Source/GLTFExporter/Private/Converters/GLTFVariantUtility.h b/Engine/Plugins/Enterprise/GLTFExporter/MigratedFromGit/Source/GLTFExporter/Private/Converters/GLTFVariantUtility.h
--- a/Engine/Plugins/Enterprise/GLTFExporter/MigratedFromGit/Source/GLTFExporter/Private/Converters/GLTFVariantUtility.h
+++ b/Engine/Plugins/Enterprise/GLTFExporter/MigratedFromGit/Source/GLTFExporter/Private/Converters/GLTFVariantUtility.h
@@ -7,6 +7,8 @@
 class FGLTFConvertBuilder;
 struct FCapturedPropSegment;
 class ULevelVariantSets;
+class UMaterialInterface;
+class UMeshComponent;
 class UVariantSet;
 class UVariant;
 class UVariantObjectBinding;
diff --git a/Engine/Plugins/Enterprise/GLTFExporter/MigratedFromGit/Source/GLTFExporter/Private/Exporters/GLTFAnimSequenceExporter.cpp b/Engine/Plugins/Enterprise/GLTFExporter/MigratedFromGit/Source/GLTFExporter/Private/Exporters/GLTFAnimSequenceExporter.cpp
--- a/Engine/Plugins/Enterprise/GLTFExporter/MigratedFromGit/Source/GLTFExporter/Private/Exporters/GLTFAnimSequenceExporter.cpp
+++ b/Engine/Plugins/Enterprise/GLTFExporter/MigratedFromGit/Source/GLTFExporter/Private/Exporters/GLTFAnimSequenceExporter.cpp
@@ -4,6 +4,7 @@
 #include "Exporters/GLTFExporterUtility.h"
 #include "Builders/GLTFContainerBuilder.h"
 #include "Animation/AnimSequence.h"
+#include "Engine/SkeletalMesh.h"
 
 UGLTFAnimSequenceExporter::UGLTFAnimSequenceExporter(const FObjectInitializer& ObjectInitializer)
 	: Super(ObjectInitializer)
diff --git a/Engine/Plugins/Enterprise/GLTFExporter/MigratedFromGit/Source/GLTFExporter/Private/Exporters/GLTFExporterUtility.cpp b/Engine/Plugins/Enterprise/GLTFExporter/MigratedFromGit/Source/GLTFExporter/Private/Exporters/GLTFExporterUtility.cpp
--- a/Engine/Plugins/Enterprise/GLTFExporter/MigratedFromGit/Source/GLTFExporter/Private/Exporters/GLTFExporterUtility.cpp
+++ b/Engine/Plugins/Enterprise/GLTFExporter/MigratedFromGit/Source/GLTFExporter/Private/Exporters/GLTFExporterUtility.cpp
@@ -3,7 +3,13 @@
 #include "Exporters/GLTFExporterUtility.h"
 #include "Materials/MaterialInstance.h"
 #include "AssetRegistry/AssetRegistryModule.h"
+#include "Animation/AnimSequence.h"
+#include "Animation/Skeleton.h"
+#include "Engine/SkeletalMesh.h"
 #include "Engine/Selection.h"
+#if WITH_EDITOR
+#include "Editor.h"
+#endif
 
 void FGLTFExporterUtility::GetSelectedActors(TSet<AActor*>& OutSelectedActors)
 {
@@ -66,7 +72,7 @@ const USkeletalMesh* FGLTFExporterUtility::FindCompatibleMesh(const USkeleton *S
 #endif
 
 	FARFilter Filter;
-	Filter.ClassNames.Add(USkeletalMesh::StaticClass()->GetFName());
+	Filter.ClassPaths.Add(USkeletalMesh::StaticClass()->GetClassPathName());
 	Filter.TagsAndValues.Add(SkeletonMemberName, FAssetData(Skeleton).GetExportTextName());
 
 	TArray<FAssetData> FilteredAssets;
diff --git a/Engine/Plugins/Enterprise/GLTFExporter/MigratedFromGit/Source/GLTFExporter/Private/Exporters/GLTFExporterUtility.h b/Engine/Plugins/Enterprise/GLTFExporter/MigratedFromGit/Source/GLTFExporter/Private/Exporters/GLTFExporterUtility.h
--- a/Engine/Plugins/Enterprise/GLTFExporter/MigratedFromGit/Source/GLTFExporter/Private/Exporters/GLTFExporterUtility.h
+++ b/Engine/Plugins/Enterprise/GLTFExporter/MigratedFromGit/Source/GLTFExporter/Private/Exporters/GLTFExporterUtility.h
@@ -2,6 +2,10 @@
 
 #pragma once
 
+#include "CoreMinimal.h"
+
+class AActor;
+class USkeletalMesh;
 class UStaticMesh;
 class UMaterialInterface;
 class UAnimSequence;
diff --git a/Engine/Plugins/Enterprise/GLTFExporter/MigratedFromGit/Source/GLTFExporter/Private/Exporters/GLTFMaterialExporter.cpp b/Engine/Plugins/Enterprise/GLTFExporter/MigratedFromGit/Source/GLTFExporter/Private/Exporters/GLTFMaterialExporter.cpp
--- a/Engine/Plugins/Enterprise/GLTFExporter/MigratedFromGit/Source/GLTFExporter/Private/Exporters/GLTFMaterialExporter.cpp
+++ b/Engine/Plugins/Enterprise/GLTFExporter/MigratedFromGit/Source/GLTFExporter/Private/Exporters/GLTFMaterialExporter.cpp
@@ -3,6 +3,7 @@
 #include "Exporters/GLTFMaterialExporter.h"
 #include "Exporters/GLTFExporterUtility.h"
 #include "Builders/GLTFContainerBuilder.h"
+#include "Engine/StaticMesh.h"
 #include "Materials/MaterialInterface.h"
 #include "UObject/ConstructorHelpers.h"
 
diff --git a/Engine/Plugins/Enterprise/GLTFExporter/MigratedFromGit/Source/GLTFExporter/Private/GLTFExporterModule.cpp b/Engine/Plugins/Enterprise/GLTFExporter/MigratedFromGit/Source/GLTFExporter/Private/GLTFExporterModule.cpp
--- a/Engine/Plugins/Enterprise/GLTFExporter/MigratedFromGit/Source/GLTFExporter/Private/GLTFExporterModule.cpp
+++ b/Engine/Plugins/Enterprise/GLTFExporter/MigratedFromGit/Source/GLTFExporter/Private/GLTFExporterModule.cpp
@@ -6,6 +6,8 @@
 #include "Utilities/GLTFProxyMaterialUtilities.h"
 #include "Actions/GLTFProxyAssetActions.h"
 #include "Interfaces/IPluginManager.h"
+#include "Misc/Paths.h"
+#include "ShaderCore.h"
 #if WITH_EDITOR
 #include "Materials/MaterialInstanceConstant.h"
 #include "Materials/MaterialInstanceDynamic.h"
diff --git a/Engine/Plugins/Enterprise/GLTFExporter/MigratedFromGit/Source/GLTFExporter/Private/Json/GLTFJsonAccessor.cpp b/Engine/Plugins/Enterprise/GLTFExporter/MigratedFromGit/Source/GLTFExporter/Private/Json/GLTFJsonAccessor.cpp
--- a/Engine/Plugins/Enterprise/GLTFExporter/MigratedFromGit/Source/GLTFExporter/Private/Json/GLTFJsonAccessor.cpp
+++ b/Engine/Plugins/Enterprise/GLTFExporter/MigratedFromGit/Source/GLTFExporter/Private/Json/GLTFJsonAccessor.cpp
@@ -1,6 +1,7 @@
 // Copyright Epic Games, Inc. All Rights Reserved.
 
 #include "Json/GLTFJsonAccessor.h"
+#include "Json/GLTFJsonBufferView.h"
 
 void FGLTFJsonAccessor::WriteObject(IGLTFJsonWriter& Writer) const
 {
diff --git a/Engine/Plugins/Enterprise/GLTFExporter/MigratedFromGit/Source/GLTFExporter/Private/Json/GLTFJsonAnimation.cpp b/Engine/Plugins/Enterprise/GLTFExporter/MigratedFromGit/Source/GLTFExporter/Private/Json/GLTFJsonAnimation.cpp
--- a/Engine/Plugins/Enterprise/GLTFExporter/MigratedFromGit/Source/GLTFExporter/Private/Json/GLTFJsonAnimation.cpp
+++ b/Engine/Plugins/Enterprise/GLTFExporter/MigratedFromGit/Source/GLTFExporter/Private/Json/GLTFJsonAnimation.cpp
@@ -1,6 +1,8 @@
 // Copyright Epic Games, Inc. All Rights Reserved.
 
 #include "Json/GLTFJsonAnimation.h"
+#include "Json/GLTFJsonNode.h"
+#include "Json/GLTFJsonAccessor.h"
 
 void FGLTFJsonAnimationChannelTarget::WriteObject(IGLTFJsonWriter& Writer) const
 {
diff --git a/Engine/Plugins/Enterprise/GLTFExporter/MigratedFromGit/Source/GLTFExporter/Private/Json/GLTFJsonBufferView.cpp b/Engine/Plugins/Enterprise/GLTFExporter/MigratedFromGit/Source/GLTFExporter/Private/Json/GLTFJsonBufferView.cpp
--- a/Engine/Plugins/Enterprise/GLTFExporter/MigratedFromGit/Source/GLTFExporter/Private/Json/GLTFJsonBufferView.cpp
+++ b/Engine/Plugins/Enterprise/GLTFExporter/MigratedFromGit/Source/GLTFExporter/Private/Json/GLTFJsonBufferView.cpp
@@ -1,6 +1,7 @@
 // Copyright Epic Games, Inc. All Rights Reserved.
 
 #include "Json/GLTFJsonBufferView.h"
+#include "Json/GLTFJsonBuffer.h"
 
 void FGLTFJsonBufferView::WriteObject(IGLTFJsonWriter& Writer) const
 {
diff --git a/Engine/Plugins/Enterprise/GLTFExporter/MigratedFromGit/Source/GLTFExporter/Private/Json/GLTFJsonKhrMaterialVariant.cpp b/Engine/Plugins/Enterprise/GLTFExporter/MigratedFromGit/Source/GLTFExporter/Private/Json/GLTFJsonKhrMaterialVariant.cpp
--- a/Engine/Plugins/Enterprise/GLTFExporter/MigratedFromGit/Source/GLTFExporter/Private/Json/GLTFJsonKhrMaterialVariant.cpp
+++ b/Engine/Plugins/Enterprise/GLTFExporter/MigratedFromGit/Source/GLTFExporter/Private/Json/GLTFJsonKhrMaterialVariant.cpp
@@ -1,6 +1,7 @@
 // Copyright Epic Games, Inc. All Rights Reserved.
 
 #include "Json/GLTFJsonKhrMaterialVariant.h"
+#include "Json/GLTFJsonMaterial.h"
 
 void FGLTFJsonKhrMaterialVariantMapping::WriteObject(IGLTFJsonWriter& Writer) const
 {
diff --git a/Engine/Plugins/Enterprise/GLTFExporter/MigratedFromGit/Source/GLTFExporter/Private/Json/GLTFJsonScene.cpp b/Engine/Plugins/Enterprise/GLTFExporter/MigratedFromGit/Source/GLTFExporter/Private/Json/GLTFJsonScene.cpp
--- a/Engine/Plugins/Enterprise/GLTFExporter/MigratedFromGit/Source/GLTFExporter/Private/Json/GLTFJsonScene.cpp
+++ b/Engine/Plugins/Enterprise/GLTFExporter/MigratedFromGit/Source/GLTFExporter/Private/Json/GLTFJsonScene.cpp
@@ -1,6 +1,8 @@
 // Copyright Epic Games, Inc. All Rights Reserved.
 
 #include "Json/GLTFJsonScene.h"
+#include "Json/GLTFJsonNode.h"
+#include "Json/GLTFJsonEpicLevelVariantSets.h"
 
 void FGLTFJsonScene::WriteObject(IGLTFJsonWriter& Writer) const
 {
diff --git a/Engine/Plugins/Enterprise/GLTFExporter/MigratedFromGit/Source/GLTFExporter/Private/Tasks/GLTFDelayedAnimationTasks.cpp b/Engine/Plugins/Enterprise/GLTFExporter/MigratedFromGit/Source/GLTFExporter/Private/Tasks/GLTFDelayedAnimationTasks.cpp
--- a/Engine/Plugins/Enterprise/GLTFExporter/MigratedFromGit/Source/GLTFExporter/Private/Tasks/GLTFDelayedAnimationTasks.cpp
+++ b/Engine/Plugins/Enterprise/GLTFExporter/MigratedFromGit/Source/GLTFExporter/Private/Tasks/GLTFDelayedAnimationTasks.cpp
@@ -12,9 +12,7 @@
 #include "Sections/MovieScene3DTransformSection.h"
 #include "Channels/MovieSceneChannelProxy.h"
 
-#if (ENGINE_MAJOR_VERSION > 4 || ENGINE_MINOR_VERSION >= 26)
 using namespace UE;
-#endif
 
 void FGLTFDelayedAnimSequenceTask::Process()
 {
diff --git a/Engine/Plugins/Enterprise/GLTFExporter/MigratedFromGit/Source/GLTFExporter/Private/Tasks/GLTFDelayedAnimationTasks.h b/Engine/Plugins/Enterprise/GLTFExporter/MigratedFromGit/Source/GLTFExporter/Private/Tasks/GLTFDelayedAnimationTasks.h
--- a/Engine/Plugins/Enterprise/GLTFExporter/MigratedFromGit/Source/GLTFExporter/Private/Tasks/GLTFDelayedAnimationTasks.h
+++ b/Engine/Plugins/Enterprise/GLTFExporter/MigratedFromGit/Source/GLTFExporter/Private/Tasks/GLTFDelayedAnimationTasks.h
@@ -5,6 +5,7 @@
 #include "Tasks/GLTFDelayedTask.h"
 #include "Builders/GLTFConvertBuilder.h"
 #include "LevelSequence.h"
+#include "Animation/AnimSequence.h"
 
 class FGLTFDelayedAnimSequenceTask : public FGLTFDelayedTask
 {
diff --git a/Engine/Plugins/Enterprise/GLTFExporter/MigratedFromGit/Source/GLTFExporter/Private/Tasks/GLTFDelayedMaterialTasks.cpp b/Engine/Plugins/Enterprise/GLTFExporter/MigratedFromGit/Source/GLTFExporter/Private/Tasks/GLTFDelayedMaterialTasks.cpp
--- a/Engine/Plugins/Enterprise/GLTFExporter/MigratedFromGit/Source/GLTFExporter/Private/Tasks/GLTFDelayedMaterialTasks.cpp
+++ b/Engine/Plugins/Enterprise/GLTFExporter/MigratedFromGit/Source/GLTFExporter/Private/Tasks/GLTFDelayedMaterialTasks.cpp
@@ -7,7 +7,6 @@
 #include "Builders/GLTFContainerBuilder.h"
 #include "Utilities/GLTFProxyMaterialUtilities.h"
 #include "Materials/GLTFProxyMaterialInfo.h"
-#if WITH_EDITOR
 #include "Materials/MaterialInstance.h"
 #include "Materials/MaterialExpressionConstant.h"
 #include "Materials/MaterialExpressionConstant2Vector.h"
@@ -17,7 +16,6 @@
 #include "Materials/MaterialExpressionVectorParameter.h"
 #include "Materials/MaterialExpressionTextureSample.h"
 #include "Materials/MaterialExpressionTextureSampleParameter2D.h"
-#endif
 
 namespace
 {
diff --git a/Engine/Plugins/Enterprise/GLTFExporter/MigratedFromGit/Source/GLTFExporter/Private/Tasks/GLTFDelayedMaterialTasks.h b/Engine/Plugins/Enterprise/GLTFExporter/MigratedFromGit/Source/GLTFExporter/Private/Tasks/GLTFDelayedMaterialTasks.h
--- a/Engine/Plugins/Enterprise/GLTFExporter/MigratedFromGit/Source/GLTFExporter/Private/Tasks/GLTFDelayedMaterialTasks.h
+++ b/Engine/Plugins/Enterprise/GLTFExporter/MigratedFromGit/Source/GLTFExporter/Private/Tasks/GLTFDelayedMaterialTasks.h
@@ -7,6 +7,7 @@
 #include "Converters/GLTFUVOverlapChecker.h"
 #include "Converters/GLTFMeshData.h"
 #include "Materials/GLTFProxyMaterialParameterInfo.h"
+#include "Materials/Material.h"
 #if WITH_EDITOR
 #include "GLTFMaterialPropertyEx.h"
 #endif
diff --git a/Engine/Plugins/Enterprise/GLTFExporter/MigratedFromGit/Source/GLTFExporter/Private/Tasks/GLTFDelayedMeshTasks.h b/Engine/Plugins/Enterprise/GLTFExporter/MigratedFromGit/Source/GLTFExporter/Private/Tasks/GLTFDelayedMeshTasks.h
--- a/Engine/Plugins/Enterprise/GLTFExporter/MigratedFromGit/Source/GLTFExporter/Private/Tasks/GLTFDelayedMeshTasks.h
+++ b/Engine/Plugins/Enterprise/GLTFExporter/MigratedFromGit/Source/GLTFExporter/Private/Tasks/GLTFDelayedMeshTasks.h
@@ -7,6 +7,11 @@
 #include "Converters/GLTFMeshSectionConverters.h"
 #include "Converters/GLTFMaterialArray.h"
 #include "Converters/GLTFNameUtility.h"
+#include "Engine/SkeletalMesh.h"
+#include "Engine/StaticMesh.h"
+
+class SkeletalMeshComponent;
+class StaticMeshComponent;
 
 class FGLTFDelayedStaticMeshTask : public FGLTFDelayedTask
 {
diff --git a/Engine/Plugins/Enterprise/GLTFExporter/MigratedFromGit/Source/GLTFExporter/Private/Utilities/GLTFTextureUtilities.cpp b/Engine/Plugins/Enterprise/GLTFExporter/MigratedFromGit/Source/GLTFExporter/Private/Utilities/GLTFTextureUtilities.cpp
--- a/Engine/Plugins/Enterprise/GLTFExporter/MigratedFromGit/Source/GLTFExporter/Private/Utilities/GLTFTextureUtilities.cpp
+++ b/Engine/Plugins/Enterprise/GLTFExporter/MigratedFromGit/Source/GLTFExporter/Private/Utilities/GLTFTextureUtilities.cpp
@@ -39,7 +39,7 @@ void FGLTFTextureUtilities::CombineMetallicAndRoughness(const UTexture* Metallic
 void FGLTFTextureUtilities::CombineTextures(const UTexture* TextureA, const FMatrix& ColorTransformA, const UTexture* TextureB, const FMatrix& ColorTransformB, const FLinearColor& BackgroundColor, UTextureRenderTarget2D* OutRenderTarget)
 {
 	FRenderTarget* RenderTarget = OutRenderTarget->GameThread_GetRenderTargetResource();
-	FCanvas Canvas(RenderTarget, nullptr, 0.0f, 0.0f, 0.0f, GMaxRHIFeatureLevel);
+	FCanvas Canvas(RenderTarget, nullptr, FGameTime(), GMaxRHIFeatureLevel);
 
 	// TODO: use different white texture depending on if target is SRGB or not
 	const FTexture* TextureResourceA = TextureA != nullptr ? TextureA->GetResource() : GetWhiteTexture();
diff --git a/Engine/Plugins/Enterprise/GLTFExporter/MigratedFromGit/Source/GLTFExporter/Public/Builders/GLTFBuilder.h b/Engine/Plugins/Enterprise/GLTFExporter/MigratedFromGit/Source/GLTFExporter/Public/Builders/GLTFBuilder.h
--- a/Engine/Plugins/Enterprise/GLTFExporter/MigratedFromGit/Source/GLTFExporter/Public/Builders/GLTFBuilder.h
+++ b/Engine/Plugins/Enterprise/GLTFExporter/MigratedFromGit/Source/GLTFExporter/Public/Builders/GLTFBuilder.h
@@ -6,6 +6,12 @@
 #include "Json/GLTFJsonEnums.h"
 #include "UObject/GCObjectScopeGuard.h"
 
+class UMaterialInterface;
+class UStaticMesh;
+class UStaticMeshComponent;
+class USkeletalMesh;
+class USkeletalMeshComponent;
+
 class GLTFEXPORTER_API FGLTFBuilder
 {
 public:
diff --git a/Engine/Plugins/Enterprise/GLTFExporter/MigratedFromGit/Source/GLTFExporter/Public/Builders/GLTFConvertBuilder.h b/Engine/Plugins/Enterprise/GLTFExporter/MigratedFromGit/Source/GLTFExporter/Public/Builders/GLTFConvertBuilder.h
--- a/Engine/Plugins/Enterprise/GLTFExporter/MigratedFromGit/Source/GLTFExporter/Public/Builders/GLTFConvertBuilder.h
+++ b/Engine/Plugins/Enterprise/GLTFExporter/MigratedFromGit/Source/GLTFExporter/Public/Builders/GLTFConvertBuilder.h
@@ -23,6 +23,8 @@
 #include "Converters/GLTFKhrVariantConverters.h"
 #include "Converters/GLTFSkySphereConverters.h"
 
+class UMeshComponent;
+
 class GLTFEXPORTER_API FGLTFConvertBuilder : public FGLTFBufferBuilder
 {
 public:
diff --git a/Engine/Plugins/Enterprise/GLTFExporter/MigratedFromGit/Source/GLTFExporter/Public/Converters/GLTFAccessorConverters.h b/Engine/Plugins/Enterprise/GLTFExporter/MigratedFromGit/Source/GLTFExporter/Public/Converters/GLTFAccessorConverters.h
--- a/Engine/Plugins/Enterprise/GLTFExporter/MigratedFromGit/Source/GLTFExporter/Public/Converters/GLTFAccessorConverters.h
+++ b/Engine/Plugins/Enterprise/GLTFExporter/MigratedFromGit/Source/GLTFExporter/Public/Converters/GLTFAccessorConverters.h
@@ -7,6 +7,11 @@
 #include "Converters/GLTFBuilderContext.h"
 #include "Converters/GLTFMeshSection.h"
 
+class FPositionVertexBuffer;
+class FColorVertexBuffer;
+class FStaticMeshVertexBuffer;
+class FSkinWeightVertexBuffer;
+
 typedef TGLTFConverter<FGLTFJsonAccessor*, const FGLTFMeshSection*, const FPositionVertexBuffer*> IGLTFPositionBufferConverter;
 typedef TGLTFConverter<FGLTFJsonAccessor*, const FGLTFMeshSection*, const FColorVertexBuffer*> IGLTFColorBufferConverter;
 typedef TGLTFConverter<FGLTFJsonAccessor*, const FGLTFMeshSection*, const FStaticMeshVertexBuffer*> IGLTFNormalBufferConverter;
diff --git a/Engine/Plugins/Enterprise/GLTFExporter/MigratedFromGit/Source/GLTFExporter/Public/Converters/GLTFAnimationConverters.h b/Engine/Plugins/Enterprise/GLTFExporter/MigratedFromGit/Source/GLTFExporter/Public/Converters/GLTFAnimationConverters.h
--- a/Engine/Plugins/Enterprise/GLTFExporter/MigratedFromGit/Source/GLTFExporter/Public/Converters/GLTFAnimationConverters.h
+++ b/Engine/Plugins/Enterprise/GLTFExporter/MigratedFromGit/Source/GLTFExporter/Public/Converters/GLTFAnimationConverters.h
@@ -6,8 +6,12 @@
 #include "Converters/GLTFConverter.h"
 #include "Converters/GLTFBuilderContext.h"
 
-class ULevelSequence;
 class ALevelSequenceActor;
+class UAnimSequence;
+class ULevel;
+class ULevelSequence;
+class USkeletalMesh;
+class USkeletalMeshComponent;
 
 typedef TGLTFConverter<FGLTFJsonAnimation*, FGLTFJsonNode*, const USkeletalMesh*, const UAnimSequence*> IGLTFAnimationConverter;
 typedef TGLTFConverter<FGLTFJsonAnimation*, FGLTFJsonNode*, const USkeletalMeshComponent*> IGLTFAnimationDataConverter;
diff --git a/Engine/Plugins/Enterprise/GLTFExporter/MigratedFromGit/Source/GLTFExporter/Public/Converters/GLTFBackdropConverters.h b/Engine/Plugins/Enterprise/GLTFExporter/MigratedFromGit/Source/GLTFExporter/Public/Converters/GLTFBackdropConverters.h
--- a/Engine/Plugins/Enterprise/GLTFExporter/MigratedFromGit/Source/GLTFExporter/Public/Converters/GLTFBackdropConverters.h
+++ b/Engine/Plugins/Enterprise/GLTFExporter/MigratedFromGit/Source/GLTFExporter/Public/Converters/GLTFBackdropConverters.h
@@ -6,6 +6,8 @@
 #include "Converters/GLTFConverter.h"
 #include "Converters/GLTFBuilderContext.h"
 
+class AActor;
+
 typedef TGLTFConverter<FGLTFJsonBackdrop*, const AActor*> IGLTFBackdropConverter;
 
 
diff --git a/Engine/Plugins/Enterprise/GLTFExporter/MigratedFromGit/Source/GLTFExporter/Public/Converters/GLTFLightConverters.h b/Engine/Plugins/Enterprise/GLTFExporter/MigratedFromGit/Source/GLTFExporter/Public/Converters/GLTFLightConverters.h
--- a/Engine/Plugins/Enterprise/GLTFExporter/MigratedFromGit/Source/GLTFExporter/Public/Converters/GLTFLightConverters.h
+++ b/Engine/Plugins/Enterprise/GLTFExporter/MigratedFromGit/Source/GLTFExporter/Public/Converters/GLTFLightConverters.h
@@ -6,6 +6,8 @@
 #include "Converters/GLTFConverter.h"
 #include "Converters/GLTFBuilderContext.h"
 
+class ULightComponent;
+
 typedef TGLTFConverter<FGLTFJsonLight*, const ULightComponent*> IGLTFLightConverter;
 
 class GLTFEXPORTER_API FGLTFLightConverter : public FGLTFBuilderContext, public IGLTFLightConverter
diff --git a/Engine/Plugins/Enterprise/GLTFExporter/MigratedFromGit/Source/GLTFExporter/Public/Converters/GLTFLightMapConverters.h b/Engine/Plugins/Enterprise/GLTFExporter/MigratedFromGit/Source/GLTFExporter/Public/Converters/GLTFLightMapConverters.h
--- a/Engine/Plugins/Enterprise/GLTFExporter/MigratedFromGit/Source/GLTFExporter/Public/Converters/GLTFLightMapConverters.h
+++ b/Engine/Plugins/Enterprise/GLTFExporter/MigratedFromGit/Source/GLTFExporter/Public/Converters/GLTFLightMapConverters.h
@@ -6,6 +6,8 @@
 #include "Converters/GLTFConverter.h"
 #include "Converters/GLTFBuilderContext.h"
 
+class UStaticMeshComponent;
+
 typedef TGLTFConverter<FGLTFJsonLightMap*, const UStaticMeshComponent*> IGLTFLightMapConverter;
 
 class GLTFEXPORTER_API FGLTFLightMapConverter : public FGLTFBuilderContext, public IGLTFLightMapConverter
diff --git a/Engine/Plugins/Enterprise/GLTFExporter/MigratedFromGit/Source/GLTFExporter/Public/Converters/GLTFMeshConverters.h b/Engine/Plugins/Enterprise/GLTFExporter/MigratedFromGit/Source/GLTFExporter/Public/Converters/GLTFMeshConverters.h
--- a/Engine/Plugins/Enterprise/GLTFExporter/MigratedFromGit/Source/GLTFExporter/Public/Converters/GLTFMeshConverters.h
+++ b/Engine/Plugins/Enterprise/GLTFExporter/MigratedFromGit/Source/GLTFExporter/Public/Converters/GLTFMeshConverters.h
@@ -8,6 +8,11 @@
 #include "Converters/GLTFMeshSectionConverters.h"
 #include "Converters/GLTFMaterialArray.h"
 
+class UStaticMesh;
+class UStaticMeshComponent;
+class USkeletalMesh;
+class USkeletalMeshComponent;
+
 typedef TGLTFConverter<FGLTFJsonMesh*, const UStaticMesh*, const UStaticMeshComponent*, FGLTFMaterialArray, int32> IGLTFStaticMeshConverter;
 typedef TGLTFConverter<FGLTFJsonMesh*, const USkeletalMesh*, const USkeletalMeshComponent*, FGLTFMaterialArray, int32> IGLTFSkeletalMeshConverter;
 
diff --git a/Engine/Plugins/Enterprise/GLTFExporter/MigratedFromGit/Source/GLTFExporter/Public/Converters/GLTFMeshData.h b/Engine/Plugins/Enterprise/GLTFExporter/MigratedFromGit/Source/GLTFExporter/Public/Converters/GLTFMeshData.h
--- a/Engine/Plugins/Enterprise/GLTFExporter/MigratedFromGit/Source/GLTFExporter/Public/Converters/GLTFMeshData.h
+++ b/Engine/Plugins/Enterprise/GLTFExporter/MigratedFromGit/Source/GLTFExporter/Public/Converters/GLTFMeshData.h
@@ -4,12 +4,19 @@
 
 #include "CoreMinimal.h"
 #include "MeshDescription.h"
+#include "SceneTypes.h"
 
 #if WITH_EDITOR
 // TODO: remove dependency to GLTFMaterialBaking in public header
 #include "GLTFMaterialBakingStructures.h"
 #endif
 
+class FLightmapResourceCluster;
+class UStaticMesh;
+class UStaticMeshComponent;
+class USkeletalMesh;
+class USkeletalMeshComponent;
+
 struct GLTFEXPORTER_API FGLTFMeshData
 {
 	FGLTFMeshData(const UStaticMesh* StaticMesh, const UStaticMeshComponent* StaticMeshComponent, int32 LODIndex);
diff --git a/Engine/Plugins/Enterprise/GLTFExporter/MigratedFromGit/Source/GLTFExporter/Public/Converters/GLTFMeshSection.h b/Engine/Plugins/Enterprise/GLTFExporter/MigratedFromGit/Source/GLTFExporter/Public/Converters/GLTFMeshSection.h
--- a/Engine/Plugins/Enterprise/GLTFExporter/MigratedFromGit/Source/GLTFExporter/Public/Converters/GLTFMeshSection.h
+++ b/Engine/Plugins/Enterprise/GLTFExporter/MigratedFromGit/Source/GLTFExporter/Public/Converters/GLTFMeshSection.h
@@ -3,8 +3,14 @@
 #pragma once
 
 #include "CoreMinimal.h"
+#include "BoneIndices.h"
 #include "Converters/GLTFIndexArray.h"
 
+class FSkeletalMeshLODRenderData;
+struct FSkelMeshRenderSection;
+struct FStaticMeshLODResources;
+struct FStaticMeshSection;
+
 struct GLTFEXPORTER_API FGLTFMeshSection
 {
 	FGLTFMeshSection(const FStaticMeshLODResources* MeshLOD, const FGLTFIndexArray& SectionIndices);
diff --git a/Engine/Plugins/Enterprise/GLTFExporter/MigratedFromGit/Source/GLTFExporter/Public/Converters/GLTFNodeConverters.h b/Engine/Plugins/Enterprise/GLTFExporter/MigratedFromGit/Source/GLTFExporter/Public/Converters/GLTFNodeConverters.h
--- a/Engine/Plugins/Enterprise/GLTFExporter/MigratedFromGit/Source/GLTFExporter/Public/Converters/GLTFNodeConverters.h
+++ b/Engine/Plugins/Enterprise/GLTFExporter/MigratedFromGit/Source/GLTFExporter/Public/Converters/GLTFNodeConverters.h
@@ -6,6 +6,12 @@
 #include "Converters/GLTFConverter.h"
 #include "Converters/GLTFBuilderContext.h"
 
+class AActor;
+class USceneComponent;
+class UStaticMesh;
+class USkeletalMesh;
+//struct FGLTFJsonNode;
+
 typedef TGLTFConverter<FGLTFJsonNode*, const AActor*> IGLTFActorConverter;
 typedef TGLTFConverter<FGLTFJsonNode*, const USceneComponent*> IGLTFComponentConverter;
 typedef TGLTFConverter<FGLTFJsonNode*, const USceneComponent*, FName> IGLTFComponentSocketConverter;
diff --git a/Engine/Plugins/Enterprise/GLTFExporter/MigratedFromGit/Source/GLTFExporter/Public/Converters/GLTFObjectArrayScopeGuard.h b/Engine/Plugins/Enterprise/GLTFExporter/MigratedFromGit/Source/GLTFExporter/Public/Converters/GLTFObjectArrayScopeGuard.h
--- a/Engine/Plugins/Enterprise/GLTFExporter/MigratedFromGit/Source/GLTFExporter/Public/Converters/GLTFObjectArrayScopeGuard.h
+++ b/Engine/Plugins/Enterprise/GLTFExporter/MigratedFromGit/Source/GLTFExporter/Public/Converters/GLTFObjectArrayScopeGuard.h
@@ -3,6 +3,7 @@
 #pragma once
 
 #include "CoreMinimal.h"
+#include "UObject/GCObject.h"
 
 template <typename ObjectType = UObject, typename = typename TEnableIf<TIsDerivedFrom<ObjectType, UObject>::Value>::Type>
 class FGLTFObjectArrayScopeGuard : public FGCObject, public TArray<ObjectType*>
diff --git a/Engine/Plugins/Enterprise/GLTFExporter/MigratedFromGit/Source/GLTFExporter/Public/Converters/GLTFSamplerConverters.h b/Engine/Plugins/Enterprise/GLTFExporter/MigratedFromGit/Source/GLTFExporter/Public/Converters/GLTFSamplerConverters.h
--- a/Engine/Plugins/Enterprise/GLTFExporter/MigratedFromGit/Source/GLTFExporter/Public/Converters/GLTFSamplerConverters.h
+++ b/Engine/Plugins/Enterprise/GLTFExporter/MigratedFromGit/Source/GLTFExporter/Public/Converters/GLTFSamplerConverters.h
@@ -6,6 +6,8 @@
 #include "Converters/GLTFConverter.h"
 #include "Converters/GLTFBuilderContext.h"
 
+class UTexture;
+
 typedef TGLTFConverter<FGLTFJsonSampler*, const UTexture*> IGLTFSamplerConverter;
 
 class GLTFEXPORTER_API FGLTFSamplerConverter : public FGLTFBuilderContext, public IGLTFSamplerConverter
diff --git a/Engine/Plugins/Enterprise/GLTFExporter/MigratedFromGit/Source/GLTFExporter/Public/Converters/GLTFSkinConverters.h b/Engine/Plugins/Enterprise/GLTFExporter/MigratedFromGit/Source/GLTFExporter/Public/Converters/GLTFSkinConverters.h
--- a/Engine/Plugins/Enterprise/GLTFExporter/MigratedFromGit/Source/GLTFExporter/Public/Converters/GLTFSkinConverters.h
+++ b/Engine/Plugins/Enterprise/GLTFExporter/MigratedFromGit/Source/GLTFExporter/Public/Converters/GLTFSkinConverters.h
@@ -6,6 +6,8 @@
 #include "Converters/GLTFConverter.h"
 #include "Converters/GLTFBuilderContext.h"
 
+class USkeletalMesh;
+
 typedef TGLTFConverter<FGLTFJsonSkin*, FGLTFJsonNode*, const USkeletalMesh*> IGLTFSkinConverter;
 
 class GLTFEXPORTER_API FGLTFSkinConverter : public FGLTFBuilderContext, public IGLTFSkinConverter
diff --git a/Engine/Plugins/Enterprise/GLTFExporter/MigratedFromGit/Source/GLTFExporter/Public/Converters/GLTFSkySphereConverters.h b/Engine/Plugins/Enterprise/GLTFExporter/MigratedFromGit/Source/GLTFExporter/Public/Converters/GLTFSkySphereConverters.h
--- a/Engine/Plugins/Enterprise/GLTFExporter/MigratedFromGit/Source/GLTFExporter/Public/Converters/GLTFSkySphereConverters.h
+++ b/Engine/Plugins/Enterprise/GLTFExporter/MigratedFromGit/Source/GLTFExporter/Public/Converters/GLTFSkySphereConverters.h
@@ -7,6 +7,8 @@
 #include "Converters/GLTFConverter.h"
 #include "Converters/GLTFBuilderContext.h"
 
+class AActor;
+class UMaterialInstance;
 struct FGLTFJsonSkySphereColorCurve;
 
 typedef TGLTFConverter<FGLTFJsonSkySphere*, const AActor*> IGLTFSkySphereConverter;
diff --git a/Engine/Plugins/Enterprise/GLTFExporter/MigratedFromGit/Source/GLTFExporter/Public/Converters/GLTFTextureConverters.h b/Engine/Plugins/Enterprise/GLTFExporter/MigratedFromGit/Source/GLTFExporter/Public/Converters/GLTFTextureConverters.h
--- a/Engine/Plugins/Enterprise/GLTFExporter/MigratedFromGit/Source/GLTFExporter/Public/Converters/GLTFTextureConverters.h
+++ b/Engine/Plugins/Enterprise/GLTFExporter/MigratedFromGit/Source/GLTFExporter/Public/Converters/GLTFTextureConverters.h
@@ -5,6 +5,12 @@
 #include "Json/GLTFJsonCore.h"
 #include "Converters/GLTFConverter.h"
 #include "Converters/GLTFBuilderContext.h"
+#include "Engine/LightMapTexture2D.h"
+#include "Engine/Texture2D.h"
+#include "Engine/TextureCube.h"
+#include "Engine/TextureRenderTarget2D.h"
+#include "Engine/TextureRenderTargetCube.h"
+#include "RHIDefinitions.h"
 
 // TODO: generalize parameter bToSRGB to SamplerType (including normalmap unpacking)
 
diff --git a/Engine/Plugins/Enterprise/GLTFExporter/MigratedFromGit/Source/GLTFExporter/Public/Converters/GLTFUVDegenerateChecker.h b/Engine/Plugins/Enterprise/GLTFExporter/MigratedFromGit/Source/GLTFExporter/Public/Converters/GLTFUVDegenerateChecker.h
--- a/Engine/Plugins/Enterprise/GLTFExporter/MigratedFromGit/Source/GLTFExporter/Public/Converters/GLTFUVDegenerateChecker.h
+++ b/Engine/Plugins/Enterprise/GLTFExporter/MigratedFromGit/Source/GLTFExporter/Public/Converters/GLTFUVDegenerateChecker.h
@@ -2,6 +2,9 @@
 
 #pragma once
 
+#include "CoreMinimal.h"
+#include "Containers/StaticArray.h"
+
 #include "Converters/GLTFConverter.h"
 #include "Converters/GLTFIndexArray.h"
 
diff --git a/Engine/Plugins/Enterprise/GLTFExporter/MigratedFromGit/Source/GLTFExporter/Public/Converters/GLTFUVOverlapChecker.h b/Engine/Plugins/Enterprise/GLTFExporter/MigratedFromGit/Source/GLTFExporter/Public/Converters/GLTFUVOverlapChecker.h
--- a/Engine/Plugins/Enterprise/GLTFExporter/MigratedFromGit/Source/GLTFExporter/Public/Converters/GLTFUVOverlapChecker.h
+++ b/Engine/Plugins/Enterprise/GLTFExporter/MigratedFromGit/Source/GLTFExporter/Public/Converters/GLTFUVOverlapChecker.h
@@ -5,6 +5,7 @@
 #include "Converters/GLTFConverter.h"
 #include "Converters/GLTFIndexArray.h"
 
+class UMaterialInterface;
 struct FMeshDescription;
 
 class GLTFEXPORTER_API FGLTFUVOverlapChecker : public TGLTFConverter<float, const FMeshDescription*, FGLTFIndexArray, int32>
diff --git a/Engine/Plugins/Enterprise/GLTFExporter/MigratedFromGit/Source/GLTFExporter/Public/Exporters/GLTFMaterialExporter.h b/Engine/Plugins/Enterprise/GLTFExporter/MigratedFromGit/Source/GLTFExporter/Public/Exporters/GLTFMaterialExporter.h
--- a/Engine/Plugins/Enterprise/GLTFExporter/MigratedFromGit/Source/GLTFExporter/Public/Exporters/GLTFMaterialExporter.h
+++ b/Engine/Plugins/Enterprise/GLTFExporter/MigratedFromGit/Source/GLTFExporter/Public/Exporters/GLTFMaterialExporter.h
@@ -5,6 +5,8 @@
 #include "Exporters/GLTFExporter.h"
 #include "GLTFMaterialExporter.generated.h"
 
+class UStaticMesh;
+
 UCLASS()
 class GLTFEXPORTER_API UGLTFMaterialExporter : public UGLTFExporter
 {
diff --git a/Engine/Plugins/Enterprise/GLTFExporter/MigratedFromGit/Source/GLTFExporter/Public/Json/GLTFJsonCore.h b/Engine/Plugins/Enterprise/GLTFExporter/MigratedFromGit/Source/GLTFExporter/Public/Json/GLTFJsonCore.h
--- a/Engine/Plugins/Enterprise/GLTFExporter/MigratedFromGit/Source/GLTFExporter/Public/Json/GLTFJsonCore.h
+++ b/Engine/Plugins/Enterprise/GLTFExporter/MigratedFromGit/Source/GLTFExporter/Public/Json/GLTFJsonCore.h
@@ -2,6 +2,8 @@
 
 #pragma once
 
+#include "CoreMinimal.h"
+
 #include "Json/GLTFJsonEnums.h"
 #include "Json/GLTFJsonValue.h"
 #include "Json/GLTFJsonArray.h"
diff --git a/Engine/Plugins/Enterprise/GLTFExporter/MigratedFromGit/Source/GLTFExporter/Public/Json/GLTFJsonValue.h b/Engine/Plugins/Enterprise/GLTFExporter/MigratedFromGit/Source/GLTFExporter/Public/Json/GLTFJsonValue.h
--- a/Engine/Plugins/Enterprise/GLTFExporter/MigratedFromGit/Source/GLTFExporter/Public/Json/GLTFJsonValue.h
+++ b/Engine/Plugins/Enterprise/GLTFExporter/MigratedFromGit/Source/GLTFExporter/Public/Json/GLTFJsonValue.h
@@ -2,6 +2,8 @@
 
 #pragma once
 
+#include "CoreMinimal.h"
+
 class IGLTFJsonWriter;
 
 struct GLTFEXPORTER_API IGLTFJsonValue
diff --git a/Engine/Plugins/Enterprise/GLTFExporter/MigratedFromGit/Source/GLTFExporter/Public/Materials/GLTFProxyMaterialParameterInfo.h b/Engine/Plugins/Enterprise/GLTFExporter/MigratedFromGit/Source/GLTFExporter/Public/Materials/GLTFProxyMaterialParameterInfo.h
--- a/Engine/Plugins/Enterprise/GLTFExporter/MigratedFromGit/Source/GLTFExporter/Public/Materials/GLTFProxyMaterialParameterInfo.h
+++ b/Engine/Plugins/Enterprise/GLTFExporter/MigratedFromGit/Source/GLTFExporter/Public/Materials/GLTFProxyMaterialParameterInfo.h
@@ -4,6 +4,11 @@
 
 #include "Utilities/GLTFProxyMaterialUtilities.h"
 
+class UMaterialInstanceConstant;
+class UMaterialInstanceDynamic;
+class UMaterialInterface;
+class UTexture;
+
 template <typename ParameterType>
 class TGLTFProxyMaterialParameterInfo
 {
diff --git a/Engine/Plugins/Enterprise/GLTFExporter/MigratedFromGit/Source/GLTFExporter/Public/Utilities/GLTFCoreUtilities.h b/Engine/Plugins/Enterprise/GLTFExporter/MigratedFromGit/Source/GLTFExporter/Public/Utilities/GLTFCoreUtilities.h
--- a/Engine/Plugins/Enterprise/GLTFExporter/MigratedFromGit/Source/GLTFExporter/Public/Utilities/GLTFCoreUtilities.h
+++ b/Engine/Plugins/Enterprise/GLTFExporter/MigratedFromGit/Source/GLTFExporter/Public/Utilities/GLTFCoreUtilities.h
@@ -2,12 +2,18 @@
 
 #pragma once
 
+#include "Animation/AnimTypes.h"
+#include "Camera/CameraTypes.h"
 #include "Core/GLTFVector.h"
 #include "Core/GLTFColor.h"
 #include "Core/GLTFMatrix.h"
 #include "Core/GLTFQuaternion.h"
 #include "Json/GLTFJsonEnums.h"
 #include "Engine/EngineTypes.h"
+#include "Engine/Texture.h"
+#include "PackedNormal.h"
+#include "RHIDefinitions.h"
+#include "SceneTypes.h"
 
 enum class EGLTFCameraControlMode : unsigned char;
 
diff --git a/Engine/Plugins/Enterprise/GLTFExporter/MigratedFromGit/Source/GLTFExporter/Public/Utilities/GLTFProxyMaterialUtilities.h b/Engine/Plugins/Enterprise/GLTFExporter/MigratedFromGit/Source/GLTFExporter/Public/Utilities/GLTFProxyMaterialUtilities.h
--- a/Engine/Plugins/Enterprise/GLTFExporter/MigratedFromGit/Source/GLTFExporter/Public/Utilities/GLTFProxyMaterialUtilities.h
+++ b/Engine/Plugins/Enterprise/GLTFExporter/MigratedFromGit/Source/GLTFExporter/Public/Utilities/GLTFProxyMaterialUtilities.h
@@ -2,11 +2,17 @@
 
 #pragma once
 
+#include "CoreMinimal.h"
 #include "Json/GLTFJsonEnums.h"
+#include "MaterialTypes.h"
+#include "UObject/ObjectMacros.h"
 
+class UMaterial;
 class UMaterialInstance;
 class UMaterialInstanceDynamic;
 class UMaterialInstanceConstant;
+class UMaterialInterface;
+class UTexture;
 
 class GLTFEXPORTER_API FGLTFProxyMaterialUtilities
 {
diff --git a/Engine/Plugins/Enterprise/GLTFExporter/MigratedFromGit/Source/GLTFMaterialAnalyzer/Private/GLTFMaterialAnalyzerModule.cpp b/Engine/Plugins/Enterprise/GLTFExporter/MigratedFromGit/Source/GLTFMaterialAnalyzer/Private/GLTFMaterialAnalyzerModule.cpp
--- a/Engine/Plugins/Enterprise/GLTFExporter/MigratedFromGit/Source/GLTFMaterialAnalyzer/Private/GLTFMaterialAnalyzerModule.cpp
+++ b/Engine/Plugins/Enterprise/GLTFExporter/MigratedFromGit/Source/GLTFMaterialAnalyzer/Private/GLTFMaterialAnalyzerModule.cpp
@@ -2,8 +2,6 @@
 
 #include "GLTFMaterialAnalyzerModule.h"
 
-#define LOCTEXT_NAMESPACE "FGLTFMaterialAnalyzerModule"
-
 void FGLTFMaterialAnalyzerModule::StartupModule()
 {
 	// This code will execute after your module is loaded into memory; the exact timing is specified in the .uplugin file per-module
@@ -15,6 +13,4 @@ void FGLTFMaterialAnalyzerModule::ShutdownModule()
 	// we call this function before unloading the module.
 }
 
-#undef LOCTEXT_NAMESPACE
-
 IMPLEMENT_MODULE(FGLTFMaterialAnalyzerModule, GLTFMaterialAnalyzer)
diff --git a/Engine/Plugins/Enterprise/GLTFExporter/MigratedFromGit/Source/GLTFMaterialAnalyzer/Private/GLTFProxyMaterialCompiler.h b/Engine/Plugins/Enterprise/GLTFExporter/MigratedFromGit/Source/GLTFMaterialAnalyzer/Private/GLTFProxyMaterialCompiler.h
--- a/Engine/Plugins/Enterprise/GLTFExporter/MigratedFromGit/Source/GLTFMaterialAnalyzer/Private/GLTFProxyMaterialCompiler.h
+++ b/Engine/Plugins/Enterprise/GLTFExporter/MigratedFromGit/Source/GLTFMaterialAnalyzer/Private/GLTFProxyMaterialCompiler.h
@@ -50,11 +50,7 @@ public:
 		return Compiler->ObjectBounds();
 	}
 
-#if (ENGINE_MAJOR_VERSION > 4 || ENGINE_MINOR_VERSION >= 26)
 	virtual int32 CustomExpression(class UMaterialExpressionCustom* Custom, int32 OutputIndex, TArray<int32>& CompiledInputs) override
-#else
-	virtual int32 CustomExpression(class UMaterialExpressionCustom* Custom, TArray<int32>& CompiledInputs) override
-#endif
 	{
 		// NOTE: because there is no overridable compiler method for ObjectLocalBounds in engine this is our only option currently
 		if (Custom->Code.Contains(TEXT("GetPrimitiveData(Parameters).LocalObjectBounds")))
@@ -62,11 +58,7 @@ public:
 			bUsesObjectLocalBounds = true;
 		}
 
-#if (ENGINE_MAJOR_VERSION > 4 || ENGINE_MINOR_VERSION >= 26)
 		return Compiler->CustomExpression(Custom, OutputIndex, CompiledInputs);
-#else
-		return Compiler->CustomExpression(Custom, CompiledInputs);
-#endif
 	}
 
 	virtual int32 PreSkinnedLocalBounds(int32 OutputIndex) override
@@ -101,12 +93,10 @@ public:
 		return Default; // ignore all non-default branches in this runtime switch
 	}
 
-#if (ENGINE_MAJOR_VERSION > 4 || ENGINE_MINOR_VERSION >= 26)
 	virtual int32 ReflectionCapturePassSwitch(int32 Default, int32 Reflection) override
 	{
 		return Default; // ignore all non-default branches in this runtime switch
 	}
-#endif
 
 	virtual int32 ReflectionAboutCustomWorldNormal(int32 CustomWorldNormal, int32 bNormalizeCustomWorldNormal) override
 	{
@@ -157,16 +147,4 @@ public:
 	{
 		return EMaterialCompilerType::MaterialProxy;
 	}
-
-#if (ENGINE_MAJOR_VERSION == 4 && ENGINE_MINOR_VERSION >= 26)
-	virtual int32 PreSkinVertexOffset() override
-	{
-		return Compiler->PreSkinVertexOffset();
-	}
-
-	virtual int32 PostSkinVertexOffset() override
-	{
-		return Compiler->PostSkinVertexOffset();
-	}
-#endif
 };
diff --git a/Engine/Plugins/Enterprise/GLTFExporter/MigratedFromGit/Source/GLTFMaterialAnalyzer/Public/GLTFMaterialAnalysis.h b/Engine/Plugins/Enterprise/GLTFExporter/MigratedFromGit/Source/GLTFMaterialAnalyzer/Public/GLTFMaterialAnalysis.h
--- a/Engine/Plugins/Enterprise/GLTFExporter/MigratedFromGit/Source/GLTFMaterialAnalyzer/Public/GLTFMaterialAnalysis.h
+++ b/Engine/Plugins/Enterprise/GLTFExporter/MigratedFromGit/Source/GLTFMaterialAnalyzer/Public/GLTFMaterialAnalysis.h
@@ -3,6 +3,7 @@
 #pragma once
 
 #include "CoreMinimal.h"
+#include "Engine/EngineTypes.h"
 
 struct FGLTFMaterialAnalysis
 {
diff --git a/Engine/Plugins/Enterprise/GLTFExporter/MigratedFromGit/Source/GLTFMaterialAnalyzer/Public/GLTFMaterialAnalyzer.h b/Engine/Plugins/Enterprise/GLTFExporter/MigratedFromGit/Source/GLTFMaterialAnalyzer/Public/GLTFMaterialAnalyzer.h
--- a/Engine/Plugins/Enterprise/GLTFExporter/MigratedFromGit/Source/GLTFMaterialAnalyzer/Public/GLTFMaterialAnalyzer.h
+++ b/Engine/Plugins/Enterprise/GLTFExporter/MigratedFromGit/Source/GLTFMaterialAnalyzer/Public/GLTFMaterialAnalyzer.h
@@ -2,19 +2,24 @@
 
 #pragma once
 
-#include "Materials/MaterialInterface.h"
+#include "CoreMinimal.h"
+#include "UObject/ObjectMacros.h"
+
 #include "GLTFMaterialAnalysis.h"
+
+#include "Materials/MaterialInstance.h"
+
 #include "GLTFMaterialAnalyzer.generated.h"
 
 UCLASS(NotBlueprintType, Transient)
-class GLTFMATERIALANALYZER_API UGLTFMaterialAnalyzer : public UMaterialInterface
+class GLTFMATERIALANALYZER_API UGLTFMaterialAnalyzer : public UMaterialInstance
 {
 	GENERATED_BODY()
 
 public:
 
 	static void AnalyzeMaterialPropertyEx(const UMaterialInterface* InMaterial, const EMaterialProperty& InProperty, const FString& InCustomOutput, FGLTFMaterialAnalysis& OutAnalysis);
-
+	
 private:
 
 	void ResetToDefaults();
diff --git a/Engine/Plugins/Enterprise/GLTFExporter/MigratedFromGit/Source/GLTFMaterialAnalyzer/Public/GLTFMaterialAnalyzerModule.h b/Engine/Plugins/Enterprise/GLTFExporter/MigratedFromGit/Source/GLTFMaterialAnalyzer/Public/GLTFMaterialAnalyzerModule.h
--- a/Engine/Plugins/Enterprise/GLTFExporter/MigratedFromGit/Source/GLTFMaterialAnalyzer/Public/GLTFMaterialAnalyzerModule.h
+++ b/Engine/Plugins/Enterprise/GLTFExporter/MigratedFromGit/Source/GLTFMaterialAnalyzer/Public/GLTFMaterialAnalyzerModule.h
@@ -3,6 +3,8 @@
 #pragma once
 
 #include "CoreMinimal.h"
+#include "Modules/ModuleInterface.h"
+#include "Modules/ModuleManager.h"
 
 class FGLTFMaterialAnalyzerModule : public IModuleInterface
 {
diff --git a/Engine/Plugins/Enterprise/GLTFExporter/MigratedFromGit/Source/GLTFMaterialBaking/Private/AddCustomAttributeHack.cpp b/Engine/Plugins/Enterprise/GLTFExporter/MigratedFromGit/Source/GLTFMaterialBaking/Private/AddCustomAttributeHack.cpp
--- a/Engine/Plugins/Enterprise/GLTFExporter/MigratedFromGit/Source/GLTFMaterialBaking/Private/AddCustomAttributeHack.cpp
+++ b/Engine/Plugins/Enterprise/GLTFExporter/MigratedFromGit/Source/GLTFMaterialBaking/Private/AddCustomAttributeHack.cpp
@@ -1,7 +1,6 @@
 // Copyright Epic Games, Inc. All Rights Reserved.
 
 #include "AddCustomAttributeHack.h"
-#include "MaterialShared.h"
 
 // NOTE: to access the raw array of custom attributes we need to use a dummy friend class
 class FMaterialAttributePropertyDetails
diff --git a/Engine/Plugins/Enterprise/GLTFExporter/MigratedFromGit/Source/GLTFMaterialBaking/Private/AddCustomAttributeHack.h b/Engine/Plugins/Enterprise/GLTFExporter/MigratedFromGit/Source/GLTFMaterialBaking/Private/AddCustomAttributeHack.h
--- a/Engine/Plugins/Enterprise/GLTFExporter/MigratedFromGit/Source/GLTFMaterialBaking/Private/AddCustomAttributeHack.h
+++ b/Engine/Plugins/Enterprise/GLTFExporter/MigratedFromGit/Source/GLTFMaterialBaking/Private/AddCustomAttributeHack.h
@@ -3,6 +3,7 @@
 #pragma once
 
 #include "CoreMinimal.h"
+#include "MaterialShared.h"
 
 // NOTE: ugly hack to register custom attribute without calling FMaterialAttributeDefinitionMap::AddCustomAttribute (which doesn't work after the AttributeDDCString has been initialized)
 void AddCustomAttributeHack(const FGuid& AttributeID, const FString& AttributeName, const FString& FunctionName, EMaterialValueType ValueType, const FVector4& DefaultValue, MaterialAttributeBlendFunction BlendFunction = nullptr);
diff --git a/Engine/Plugins/Enterprise/GLTFExporter/MigratedFromGit/Source/GLTFMaterialBaking/Private/GLTFMaterialBakingModule.cpp b/Engine/Plugins/Enterprise/GLTFExporter/MigratedFromGit/Source/GLTFMaterialBaking/Private/GLTFMaterialBakingModule.cpp
--- a/Engine/Plugins/Enterprise/GLTFExporter/MigratedFromGit/Source/GLTFMaterialBaking/Private/GLTFMaterialBakingModule.cpp
+++ b/Engine/Plugins/Enterprise/GLTFExporter/MigratedFromGit/Source/GLTFMaterialBaking/Private/GLTFMaterialBakingModule.cpp
@@ -13,6 +13,7 @@
 #include "Materials/MaterialInstance.h"
 #include "Materials/MaterialInstanceConstant.h"
 #include "MaterialEditor/MaterialEditorInstanceConstant.h"
+#include "Modules/ModuleManager.h"
 #include "RenderingThread.h"
 #include "RHISurfaceDataConversion.h"
 #include "Misc/ScopedSlowTask.h"
diff --git a/Engine/Plugins/Enterprise/GLTFExporter/MigratedFromGit/Source/GLTFMaterialBaking/Private/MaterialRenderItem.cpp b/Engine/Plugins/Enterprise/GLTFExporter/MigratedFromGit/Source/GLTFMaterialBaking/Private/MaterialRenderItem.cpp
--- a/Engine/Plugins/Enterprise/GLTFExporter/MigratedFromGit/Source/GLTFMaterialBaking/Private/MaterialRenderItem.cpp
+++ b/Engine/Plugins/Enterprise/GLTFExporter/MigratedFromGit/Source/GLTFMaterialBaking/Private/MaterialRenderItem.cpp
@@ -55,7 +55,7 @@ bool FGLTFMeshMaterialRenderItem::Render_RenderThread(FCanvasRenderContext& Rend
 
 	if (Vertices.Num() && Indices.Num())
 	{
-		FMeshPassProcessorRenderState LocalDrawRenderState(View);
+		FMeshPassProcessorRenderState LocalDrawRenderState;
 
 		// disable depth test & writes
 		LocalDrawRenderState.SetBlendState(TStaticBlendState<CW_RGBA>::GetRHI());
diff --git a/Engine/Plugins/Enterprise/GLTFExporter/MigratedFromGit/Source/GLTFMaterialBaking/Public/GLTFMaterialBakingHelpers.h b/Engine/Plugins/Enterprise/GLTFExporter/MigratedFromGit/Source/GLTFMaterialBaking/Public/GLTFMaterialBakingHelpers.h
--- a/Engine/Plugins/Enterprise/GLTFExporter/MigratedFromGit/Source/GLTFMaterialBaking/Public/GLTFMaterialBakingHelpers.h
+++ b/Engine/Plugins/Enterprise/GLTFExporter/MigratedFromGit/Source/GLTFMaterialBaking/Public/GLTFMaterialBakingHelpers.h
@@ -3,8 +3,7 @@
 #pragma once
 
 #include "Containers/Array.h"
-
-struct FColor;
+#include "Math/Color.h"
 
 class GLTFMATERIALBAKING_API FGLTFMaterialBakingHelpers
 {
diff --git a/Engine/Plugins/Enterprise/GLTFExporter/MigratedFromGit/Source/GLTFMaterialBaking/Public/GLTFMaterialBakingStructures.h b/Engine/Plugins/Enterprise/GLTFExporter/MigratedFromGit/Source/GLTFMaterialBaking/Public/GLTFMaterialBakingStructures.h
--- a/Engine/Plugins/Enterprise/GLTFExporter/MigratedFromGit/Source/GLTFMaterialBaking/Public/GLTFMaterialBakingStructures.h
+++ b/Engine/Plugins/Enterprise/GLTFExporter/MigratedFromGit/Source/GLTFMaterialBaking/Public/GLTFMaterialBakingStructures.h
@@ -3,11 +3,16 @@
 #pragma once
 
 #include "Math/IntPoint.h"
+#include "Components/MeshComponent.h"
+#include "Components/SkinnedMeshComponent.h"
+#include "Engine/SkeletalMesh.h"
+#include "Engine/StaticMesh.h"
 #include "SceneTypes.h"
 #include "LightMap.h"
 #include "GLTFMaterialPropertyEx.h"
 
 class UMaterialInterface;
+struct FCustomPrimitiveData;
 struct FMeshDescription;
 
 /** Structure containing information about the material which is being baked out */
