[GLTFExporter] Move GLTFExporter into  //UE5/Main/Engine/Plugins/Enterprise/

commit 32cd9bfcccb9c3366948fe02bcefa9a14e3a346a
Date: Tue, 19 Jul 2022 11:01:36 -0400

    Move GLTFExporter into  //UE5/Main/Engine/Plugins/Enterprise/
    jira none
    fyi aidin.abedi, flavien.picon, francis.moreau
    preflight 62d6ab27185da2495f96dfb3
    okforgithub public, internal

#skipci

diff --git a/Engine/Plugins/Enterprise/GLTFExporter/MigratedFromGit/Source/GLTFExporter/GLTFExporter.Build.cs b/Engine/Plugins/Enterprise/GLTFExporter/MigratedFromGit/Source/GLTFExporter/GLTFExporter.Build.cs
--- a/Engine/Plugins/Enterprise/GLTFExporter/MigratedFromGit/Source/GLTFExporter/GLTFExporter.Build.cs
+++ b/Engine/Plugins/Enterprise/GLTFExporter/MigratedFromGit/Source/GLTFExporter/GLTFExporter.Build.cs
@@ -6,8 +6,9 @@ public class GLTFExporter : ModuleRules
 {
 	public GLTFExporter(ReadOnlyTargetRules Target) : base(Target)
 	{
-		PCHUsage = PCHUsageMode.UseExplicitOrSharedPCHs;
-		bTreatAsEngineModule = true; // Only necessary when plugin installed in project
+		//PCHUsage = PCHUsageMode.UseExplicitOrSharedPCHs;
+		PCHUsage = PCHUsageMode.NoPCHs;
+		bUseUnity = false;
 
 		PublicDependencyModuleNames.AddRange(
 			new string[]
diff --git a/Engine/Plugins/Enterprise/GLTFExporter/MigratedFromGit/Source/GLTFMaterialAnalyzer/GLTFMaterialAnalyzer.Build.cs b/Engine/Plugins/Enterprise/GLTFExporter/MigratedFromGit/Source/GLTFMaterialAnalyzer/GLTFMaterialAnalyzer.Build.cs
--- a/Engine/Plugins/Enterprise/GLTFExporter/MigratedFromGit/Source/GLTFMaterialAnalyzer/GLTFMaterialAnalyzer.Build.cs
+++ b/Engine/Plugins/Enterprise/GLTFExporter/MigratedFromGit/Source/GLTFMaterialAnalyzer/GLTFMaterialAnalyzer.Build.cs
@@ -6,10 +6,11 @@ public class GLTFMaterialAnalyzer : ModuleRules
 {
 	public GLTFMaterialAnalyzer(ReadOnlyTargetRules Target) : base(Target)
 	{
-		PCHUsage = PCHUsageMode.UseExplicitOrSharedPCHs;
-		bTreatAsEngineModule = true; // Only necessary when plugin installed in project
+		//PCHUsage = PCHUsageMode.UseExplicitOrSharedPCHs;
+		PCHUsage = PCHUsageMode.NoPCHs;
+		bUseUnity = false;
 
-		PublicDependencyModuleNames .AddRange(
+		PublicDependencyModuleNames.AddRange(
 			new string[]
 			{
 				"Core",
diff --git a/Engine/Plugins/Enterprise/GLTFExporter/MigratedFromGit/Source/GLTFMaterialBaking/GLTFMaterialBaking.Build.cs b/Engine/Plugins/Enterprise/GLTFExporter/MigratedFromGit/Source/GLTFMaterialBaking/GLTFMaterialBaking.Build.cs
--- a/Engine/Plugins/Enterprise/GLTFExporter/MigratedFromGit/Source/GLTFMaterialBaking/GLTFMaterialBaking.Build.cs
+++ b/Engine/Plugins/Enterprise/GLTFExporter/MigratedFromGit/Source/GLTFMaterialBaking/GLTFMaterialBaking.Build.cs
@@ -6,10 +6,11 @@ public class GLTFMaterialBaking : ModuleRules
 {
 	public GLTFMaterialBaking(ReadOnlyTargetRules Target) : base(Target)
 	{
-		PCHUsage = PCHUsageMode.UseExplicitOrSharedPCHs;
-		bTreatAsEngineModule = true; // Only necessary when plugin installed in project
+		//PCHUsage = PCHUsageMode.UseExplicitOrSharedPCHs;
+		PCHUsage = PCHUsageMode.NoPCHs;
+		bUseUnity = false;
 
-		PublicDependencyModuleNames .AddRange(
+		PublicDependencyModuleNames.AddRange(
 			new string[]
 			{
 				"Core",
