[GLTFExporter] Skeletal Mesh - Replace the newly created GetSkeletalMesh function by GetSkeletalMeshAsset so it matches the setter SetSkeletalMeshAsset function (which itself couldn't be named SetSkeletalMesh due to the function already existing and doing something else).

commit bec406d6c97ce10f6348b6d5145ad0a7fcb94357
Date: Mon, 15 Aug 2022 09:26:50 -0400

    Skeletal Mesh - Replace the newly created GetSkeletalMesh function by GetSkeletalMeshAsset so it matches the setter SetSkeletalMeshAsset function (which itself couldn't be named SetSkeletalMesh due to the function already existing and doing something else).
    preflight 62fa2afeae3edb54c979492e
    jira none
    okforgithub public, internal

#skipci

diff --git a/Engine/Plugins/Enterprise/GLTFExporter/MigratedFromGit/Source/GLTFExporter/Private/Actors/GLTFHotspotActor.cpp b/Engine/Plugins/Enterprise/GLTFExporter/MigratedFromGit/Source/GLTFExporter/Private/Actors/GLTFHotspotActor.cpp
--- a/Engine/Plugins/Enterprise/GLTFExporter/MigratedFromGit/Source/GLTFExporter/Private/Actors/GLTFHotspotActor.cpp
+++ b/Engine/Plugins/Enterprise/GLTFExporter/MigratedFromGit/Source/GLTFExporter/Private/Actors/GLTFHotspotActor.cpp
@@ -453,7 +453,7 @@ void AGLTFHotspotActor::ValidateAnimation()
 {
 	if (SkeletalMeshActor != nullptr && AnimationSequence != nullptr)
 	{
-		if (const USkeletalMesh* SkeletalMesh = SkeletalMeshActor->GetSkeletalMeshComponent()->GetSkeletalMesh())
+		if (const USkeletalMesh* SkeletalMesh = SkeletalMeshActor->GetSkeletalMeshComponent()->GetSkeletalMeshAsset())
 		{
 #if (ENGINE_MAJOR_VERSION > 4 || ENGINE_MINOR_VERSION >= 27)
 			const USkeleton* Skeleton = SkeletalMesh->GetSkeleton();
diff --git a/Engine/Plugins/Enterprise/GLTFExporter/MigratedFromGit/Source/GLTFExporter/Private/Builders/GLTFConvertBuilder.cpp b/Engine/Plugins/Enterprise/GLTFExporter/MigratedFromGit/Source/GLTFExporter/Private/Builders/GLTFConvertBuilder.cpp
--- a/Engine/Plugins/Enterprise/GLTFExporter/MigratedFromGit/Source/GLTFExporter/Private/Builders/GLTFConvertBuilder.cpp
+++ b/Engine/Plugins/Enterprise/GLTFExporter/MigratedFromGit/Source/GLTFExporter/Private/Builders/GLTFConvertBuilder.cpp
@@ -163,7 +163,7 @@ FGLTFJsonMesh* FGLTFConvertBuilder::AddUniqueMesh(const USkeletalMeshComponent*
 		return nullptr;
 	}
 
-	const USkeletalMesh* SkeletalMesh = SkeletalMeshComponent->GetSkeletalMesh();
+	const USkeletalMesh* SkeletalMesh = SkeletalMeshComponent->GetSkeletalMeshAsset();
 	if (SkeletalMesh == nullptr)
 	{
 		return nullptr;
@@ -225,7 +225,7 @@ FGLTFJsonMaterial* FGLTFConvertBuilder::AddUniqueMaterial(const UMaterialInterfa
 FGLTFJsonMaterial* FGLTFConvertBuilder::AddUniqueMaterial(const UMaterialInterface* Material, const USkeletalMeshComponent* SkeletalMeshComponent, int32 LODIndex, int32 MaterialIndex)
 {
 	// TODO: optimize by skipping mesh data if material doesn't need it
-	const USkeletalMesh* SkeletalMesh = SkeletalMeshComponent->GetSkeletalMesh();
+	const USkeletalMesh* SkeletalMesh = SkeletalMeshComponent->GetSkeletalMeshAsset();
 	const FGLTFMeshData* MeshData = AddUniqueMeshData(SkeletalMesh, SkeletalMeshComponent, LODIndex);
 	const FGLTFIndexArray SectionIndices = FGLTFMeshUtility::GetSectionIndices(SkeletalMesh, MeshData->LODIndex, MaterialIndex);
 	return AddUniqueMaterial(Material, MeshData, SectionIndices);
@@ -363,7 +363,7 @@ FGLTFJsonSkin* FGLTFConvertBuilder::AddUniqueSkin(FGLTFJsonNode* RootNode, const
 		return nullptr;
 	}
 
-	const USkeletalMesh* SkeletalMesh = SkeletalMeshComponent->GetSkeletalMesh();
+	const USkeletalMesh* SkeletalMesh = SkeletalMeshComponent->GetSkeletalMeshAsset();
 
 	if (SkeletalMesh == nullptr)
 	{
diff --git a/Engine/Plugins/Enterprise/GLTFExporter/MigratedFromGit/Source/GLTFExporter/Private/Converters/GLTFAnimationConverters.cpp b/Engine/Plugins/Enterprise/GLTFExporter/MigratedFromGit/Source/GLTFExporter/Private/Converters/GLTFAnimationConverters.cpp
--- a/Engine/Plugins/Enterprise/GLTFExporter/MigratedFromGit/Source/GLTFExporter/Private/Converters/GLTFAnimationConverters.cpp
+++ b/Engine/Plugins/Enterprise/GLTFExporter/MigratedFromGit/Source/GLTFExporter/Private/Converters/GLTFAnimationConverters.cpp
@@ -40,7 +40,7 @@ FGLTFJsonAnimation* FGLTFAnimationConverter::Convert(FGLTFJsonNode* RootNode, co
 
 FGLTFJsonAnimation* FGLTFAnimationDataConverter::Convert(FGLTFJsonNode* RootNode, const USkeletalMeshComponent* SkeletalMeshComponent)
 {
-	const USkeletalMesh* SkeletalMesh = SkeletalMeshComponent->GetSkeletalMesh();
+	const USkeletalMesh* SkeletalMesh = SkeletalMeshComponent->GetSkeletalMeshAsset();
 	if (SkeletalMesh == nullptr)
 	{
 		return nullptr;
diff --git a/Engine/Plugins/Enterprise/GLTFExporter/MigratedFromGit/Source/GLTFExporter/Private/Converters/GLTFHotspotConverters.cpp b/Engine/Plugins/Enterprise/GLTFExporter/MigratedFromGit/Source/GLTFExporter/Private/Converters/GLTFHotspotConverters.cpp
--- a/Engine/Plugins/Enterprise/GLTFExporter/MigratedFromGit/Source/GLTFExporter/Private/Converters/GLTFHotspotConverters.cpp
+++ b/Engine/Plugins/Enterprise/GLTFExporter/MigratedFromGit/Source/GLTFExporter/Private/Converters/GLTFHotspotConverters.cpp
@@ -29,7 +29,7 @@ FGLTFJsonHotspot* FGLTFHotspotConverter::Convert(const AGLTFHotspotActor* Hotspo
 			FGLTFJsonNode* RootNode = Builder.AddUniqueNode(SkeletalMeshActor);
 
 			const USkeletalMeshComponent* SkeletalMeshComponent = SkeletalMeshActor->GetSkeletalMeshComponent();
-			if (const USkeletalMesh* SkeletalMesh = SkeletalMeshComponent->GetSkeletalMesh())
+			if (const USkeletalMesh* SkeletalMesh = SkeletalMeshComponent->GetSkeletalMeshAsset())
 			{
 				if (const UAnimSequence* AnimSequence = HotspotActor->AnimationSequence)
 				{
