[GLTFExporter] ObjectPtr upgrade for engine plugins used from UnrealEditor -allmodules.  Automated upgrade process.

commit cc5a3da731d70491718c7718abf3e2946d64891c
Date: Thu, 18 Aug 2022 15:08:49 -0400

    ObjectPtr upgrade for engine plugins used from UnrealEditor -allmodules.  Automated upgrade process.
    preflight 62fe7f270601ad0504910c30
    okforgithub public, dpx, internal

#skipci

diff --git a/Engine/Plugins/Enterprise/GLTFExporter/MigratedFromGit/Source/GLTFExporter/Public/Actors/GLTFCameraActor.h b/Engine/Plugins/Enterprise/GLTFExporter/MigratedFromGit/Source/GLTFExporter/Public/Actors/GLTFCameraActor.h
--- a/Engine/Plugins/Enterprise/GLTFExporter/MigratedFromGit/Source/GLTFExporter/Public/Actors/GLTFCameraActor.h
+++ b/Engine/Plugins/Enterprise/GLTFExporter/MigratedFromGit/Source/GLTFExporter/Public/Actors/GLTFCameraActor.h
@@ -71,7 +71,7 @@ public:
 
 	/* Actor which the camera will focus on and subsequently orbit when using Orbital mode. */
 	UPROPERTY(EditAnywhere, BlueprintReadWrite, Category = "GLTF Camera Actor", meta=(EditCondition="Mode == EGLTFCameraControlMode::Orbital"))
-	AActor* Target;
+	TObjectPtr<AActor> Target;
 
 	/* Minimum pitch angle (in degrees) for the camera. */
 	UPROPERTY(EditAnywhere, BlueprintReadWrite, Category = "GLTF Camera Actor")
diff --git a/Engine/Plugins/Enterprise/GLTFExporter/MigratedFromGit/Source/GLTFExporter/Public/Actors/GLTFHotspotActor.h b/Engine/Plugins/Enterprise/GLTFExporter/MigratedFromGit/Source/GLTFExporter/Public/Actors/GLTFHotspotActor.h
--- a/Engine/Plugins/Enterprise/GLTFExporter/MigratedFromGit/Source/GLTFExporter/Public/Actors/GLTFHotspotActor.h
+++ b/Engine/Plugins/Enterprise/GLTFExporter/MigratedFromGit/Source/GLTFExporter/Public/Actors/GLTFHotspotActor.h
@@ -62,63 +62,63 @@ public:
 
 	/* The skeletal mesh actor that will be animated when the hotspot is clicked. */
 	UPROPERTY(EditAnywhere, BlueprintReadWrite, Category = "Hotspot Animation")
-	ASkeletalMeshActor* SkeletalMeshActor;
+	TObjectPtr<ASkeletalMeshActor> SkeletalMeshActor;
 
 	/* The animation sequence that will be played on the skeletal mesh actor. Must be compatible with its skeletal mesh asset. */
 	UPROPERTY(EditAnywhere, BlueprintReadWrite, Category = "Hotspot Animation", meta=(EditCondition="SkeletalMeshActor != nullptr"))
-	UAnimSequence* AnimationSequence;
+	TObjectPtr<UAnimSequence> AnimationSequence;
 
 	/* The level sequence that will be played in the level. */
 	UPROPERTY(EditAnywhere, BlueprintReadWrite, Category = "Hotspot Animation", meta=(EditCondition="SkeletalMeshActor == nullptr"))
-	ULevelSequence* LevelSequence;
+	TObjectPtr<ULevelSequence> LevelSequence;
 
 	/* The billboard image that will be shown when the hotspot is in an inactive state or one without a specified image. */
 	UPROPERTY(EditAnywhere, BlueprintReadWrite, Category = "Hotspot Appearance")
-	UTexture2D* Image;
+	TObjectPtr<UTexture2D> Image;
 
 	/** The optional billboard image that will be shown when a cursor enters the hotspot. */
 	UPROPERTY(EditAnywhere, BlueprintReadWrite, Category = "Hotspot Appearance", meta=(EditCondition="Image != nullptr"))
-	UTexture2D* HoveredImage;
+	TObjectPtr<UTexture2D> HoveredImage;
 
 	/** The optional billboard image that will be shown when the hotspot is toggled by a click. */
 	UPROPERTY(EditAnywhere, BlueprintReadWrite, Category = "Hotspot Appearance", meta=(EditCondition="Image != nullptr"))
-	UTexture2D* ToggledImage;
+	TObjectPtr<UTexture2D> ToggledImage;
 
 	/** The optional billboard image that will be shown when the hotspot is toggled by a click and a cursor enters it. */
 	UPROPERTY(EditAnywhere, BlueprintReadWrite, Category = "Hotspot Appearance", meta=(EditCondition="Image != nullptr"))
-	UTexture2D* ToggledHoveredImage;
+	TObjectPtr<UTexture2D> ToggledHoveredImage;
 
 private:
 
 	UPROPERTY(Transient, DuplicateTransient)
-	UMaterialBillboardComponent* BillboardComponent;
+	TObjectPtr<UMaterialBillboardComponent> BillboardComponent;
 
 	UPROPERTY(Transient, DuplicateTransient)
-	USphereComponent* SphereComponent;
+	TObjectPtr<USphereComponent> SphereComponent;
 
 	UPROPERTY(Transient, DuplicateTransient)
-	UMaterialInterface* DefaultMaterial;
+	TObjectPtr<UMaterialInterface> DefaultMaterial;
 
 	UPROPERTY(Transient, DuplicateTransient)
-	UTexture2D* DefaultImage;
+	TObjectPtr<UTexture2D> DefaultImage;
 
 	UPROPERTY(Transient, DuplicateTransient)
-	UTexture2D* DefaultHoveredImage;
+	TObjectPtr<UTexture2D> DefaultHoveredImage;
 
 	UPROPERTY(Transient, DuplicateTransient)
-	UTexture2D* DefaultToggledImage;
+	TObjectPtr<UTexture2D> DefaultToggledImage;
 
 	UPROPERTY(Transient, DuplicateTransient)
-	UTexture2D* DefaultToggledHoveredImage;
+	TObjectPtr<UTexture2D> DefaultToggledHoveredImage;
 
 	UPROPERTY(Transient, DuplicateTransient)
-	UMaterialInterface* DefaultIconMaterial;
+	TObjectPtr<UMaterialInterface> DefaultIconMaterial;
 
 	UPROPERTY(Transient, DuplicateTransient)
-	const UTexture* ActiveImage;
+	TObjectPtr<const UTexture> ActiveImage;
 
 	UPROPERTY(Instanced, Transient)
-	ULevelSequencePlayer* LevelSequencePlayer;
+	TObjectPtr<ULevelSequencePlayer> LevelSequencePlayer;
 
 	FVector2D ActiveImageSize;
 
diff --git a/Engine/Plugins/Enterprise/GLTFExporter/MigratedFromGit/Source/GLTFExporter/Public/Exporters/GLTFMaterialExporter.h b/Engine/Plugins/Enterprise/GLTFExporter/MigratedFromGit/Source/GLTFExporter/Public/Exporters/GLTFMaterialExporter.h
--- a/Engine/Plugins/Enterprise/GLTFExporter/MigratedFromGit/Source/GLTFExporter/Public/Exporters/GLTFMaterialExporter.h
+++ b/Engine/Plugins/Enterprise/GLTFExporter/MigratedFromGit/Source/GLTFExporter/Public/Exporters/GLTFMaterialExporter.h
@@ -21,5 +21,5 @@ public:
 private:
 
 	UPROPERTY(Transient, DuplicateTransient)
-	UStaticMesh* DefaultPreviewMesh;
+	TObjectPtr<UStaticMesh> DefaultPreviewMesh;
 };
diff --git a/Engine/Plugins/Enterprise/GLTFExporter/MigratedFromGit/Source/GLTFExporter/Public/UserData/GLTFMaterialUserData.h b/Engine/Plugins/Enterprise/GLTFExporter/MigratedFromGit/Source/GLTFExporter/Public/UserData/GLTFMaterialUserData.h
--- a/Engine/Plugins/Enterprise/GLTFExporter/MigratedFromGit/Source/GLTFExporter/Public/UserData/GLTFMaterialUserData.h
+++ b/Engine/Plugins/Enterprise/GLTFExporter/MigratedFromGit/Source/GLTFExporter/Public/UserData/GLTFMaterialUserData.h
@@ -96,7 +96,7 @@ public:
 
 	/** If assigned, export will use the proxy instead of the original material. */
 	UPROPERTY(EditAnywhere, BlueprintReadWrite, Category = "General")
-	UMaterialInterface* Proxy;
+	TObjectPtr<UMaterialInterface> Proxy;
 
 	/** Default bake settings for this material in general. */
 	UPROPERTY(EditAnywhere, BlueprintReadWrite, Category = "Override Bake Settings", meta = (ShowOnlyInnerProperties))
